NeuralNetworkInference plugin:
- QA Asset moved from Public into Internal.

#rnx

diff --git a/Engine/Plugins/Experimental/NeuralNetworkInference/Source/NeuralNetworkInferenceQA/Internal/NeuralNetworkInferenceQAAsset.h b/Engine/Plugins/Experimental/NeuralNetworkInference/Source/NeuralNetworkInferenceQA/Internal/NeuralNetworkInferenceQAAsset.h
new file mode 100644
--- /dev/null
+++ b/Engine/Plugins/Experimental/NeuralNetworkInference/Source/NeuralNetworkInferenceQA/Internal/NeuralNetworkInferenceQAAsset.h
@@ -0,0 +1,107 @@
+// Copyright Epic Games, Inc. All Rights Reserved.
+
+#pragma once
+
+#include "CoreMinimal.h"
+#include "NeuralTensor.h"
+#include "NeuralNetworkInferenceQAAsset.generated.h"
+
+class UNeuralNetworkLegacy;
+
+
+
+USTRUCT(BlueprintType)
+struct NEURALNETWORKINFERENCEQA_API FNeuralNetworkInferenceQAOperatorTestAsset
+{
+	GENERATED_BODY()
+
+public:
+	UPROPERTY(VisibleAnywhere, Category = "Neural Network Inference QA")
+	TArray<FNeuralTensor> InputTensors;
+	
+	UPROPERTY(VisibleAnywhere, Category = "Neural Network Inference QA")
+	TArray<FNeuralTensor> OutputTensors;
+
+	/**
+	 * It adds the InNameIndexMap subset of InTensors into Tensors.
+	 */
+	void AddInputTensors(const TArray<FNeuralTensor>& InTensors, const TMap<FString, int32>& InNameIndexMap);
+	void AddOutputTensors(const TArray<FNeuralTensor>& InTensors, const TMap<FString, int32>& InNameIndexMap);
+
+	/**
+	 * It will compare new vs. previous tests.
+	 */
+	static bool CompareAverageL1DiffNewVsPreviousTests(const FNeuralNetworkInferenceQAOperatorTestAsset& InNeuralNetworkInferenceQAOperatorTestAsset1, const FNeuralNetworkInferenceQAOperatorTestAsset& InNeuralNetworkInferenceQAOperatorTestAsset2, const FString& InOperatorName);
+
+private:
+	/**
+	 * Auxiliary for CompareAverageL1DiffNewVsPreviousTests(const FNeuralNetworkInferenceQAOperatorTestAsset&, const FNeuralNetworkInferenceQAOperatorTestAsset&).
+	 */
+	static bool CompareAverageL1DiffNewVsPreviousTests(const TArray<FNeuralTensor>& InTensorsA, const TArray<FNeuralTensor>& InTensorsB, const FString& InOperatorName);
+};
+
+
+
+USTRUCT(BlueprintType)
+struct NEURALNETWORKINFERENCEQA_API FNeuralNetworkInferenceQAOperatorAsset
+{
+	GENERATED_BODY()
+
+public:
+	static FString RunNetworkCPUAndGetString(UNeuralNetworkLegacy* InOutNetwork);
+
+	void RunAndAddTest(UNeuralNetworkLegacy* InOutNetwork);
+
+	/**
+	 * It will compare new vs. previous tests.
+	 */
+	bool CompareNewVsPreviousTests(const FString& InGroundTruthDirectory, const FString& InOperatorName);
+
+	/**
+	 * It flushes (cleans/resets) the new results so the tests can be run again.
+	 */
+	void FlushNewTests();
+	
+protected:
+	FString NewTestsString;
+
+	TArray<FNeuralNetworkInferenceQAOperatorTestAsset> NewTests;
+
+	UPROPERTY(VisibleAnywhere, Category = "Neural Network Inference QA")
+	TArray<FNeuralNetworkInferenceQAOperatorTestAsset> PreviousTests;
+
+	UPROPERTY()
+	FString PreviousTestsString;
+};
+
+
+
+/**
+ * This is an auxiliary class. See UNeuralNetworkLegacy for a high-level wrapper of the whole NeuralNetworkInference plugin. The UNeuralNetworkLegacy header
+ * documentation also includes some code examples.
+ */
+UCLASS(BlueprintType)
+class NEURALNETWORKINFERENCEQA_API UNeuralNetworkInferenceQAAsset : public UObject
+{
+	GENERATED_BODY()
+
+public:
+	void FindOrAddOperators(const TArray<FString>& InOperatorNames);
+
+	void RunAndAddTest(UNeuralNetworkLegacy* InOutNetwork, const FString& OperatorName);
+
+	bool CompareNewVsPreviousTests(const FString& InGroundTruthDirectory);
+
+	/**
+	 * It resets the new results so it can be run again.
+	 */
+	void FlushNewTests();
+
+	static UNeuralNetworkInferenceQAAsset* Load(const FString& InNeuralNetworkInferenceQAAssetParentDirectoryName, const FString& InNeuralNetworkInferenceQAAssetName);
+
+	bool Save();
+
+protected:
+	UPROPERTY(EditAnywhere, Category = "Neural Network Inference QA")
+	TMap<FString, FNeuralNetworkInferenceQAOperatorAsset> Operators;
+};
diff --git a/Engine/Plugins/Experimental/NeuralNetworkInference/Source/NeuralNetworkInferenceQA/Public/NeuralNetworkInferenceQAAsset.h b/Engine/Plugins/Experimental/NeuralNetworkInference/Source/NeuralNetworkInferenceQA/Public/NeuralNetworkInferenceQAAsset.h
deleted file mode 100644
