NeuralNetworkInference plugin:
- Which OS will compile is longer set on the UPlugin, but instead on the NNI.Build.Cs file, by adding WITH_ORT flag for desired platforms. Simpler for users.

Preflight https://horde.devtools.epicgames.com/job/6114536e9c7bb10001fbd906
#rnx

diff --git a/Engine/Plugins/Experimental/NeuralNetworkInference/NeuralNetworkInference.uplugin b/Engine/Plugins/Experimental/NeuralNetworkInference/NeuralNetworkInference.uplugin
--- a/Engine/Plugins/Experimental/NeuralNetworkInference/NeuralNetworkInference.uplugin
+++ b/Engine/Plugins/Experimental/NeuralNetworkInference/NeuralNetworkInference.uplugin
@@ -18,10 +18,7 @@
 		{
 			"Name": "NeuralNetworkInference",
 			"Type": "RuntimeAndProgram",
-			"LoadingPhase": "Default",
-			"WhitelistPlatforms": [
-				"Win64"
-			]
+			"LoadingPhase": "Default"
 		},
 		{
 			"Name": "NeuralNetworkInferenceBackEnd",
diff --git a/Engine/Plugins/Experimental/NeuralNetworkInference/Source/NeuralNetworkInference/NeuralNetworkInference.Build.cs b/Engine/Plugins/Experimental/NeuralNetworkInference/Source/NeuralNetworkInference/NeuralNetworkInference.Build.cs
--- a/Engine/Plugins/Experimental/NeuralNetworkInference/Source/NeuralNetworkInference/NeuralNetworkInference.Build.cs
+++ b/Engine/Plugins/Experimental/NeuralNetworkInference/Source/NeuralNetworkInference/NeuralNetworkInference.Build.cs
@@ -6,7 +6,14 @@ public class NeuralNetworkInference : ModuleRules
 {
 	public NeuralNetworkInference( ReadOnlyTargetRules Target ) : base( Target )
 	{
-        ShortName = "DeprecatedNNIWE"; // Could be removed when plugin moves to Experimental, NFL path is too long
+		// Define when ORT-based NNI is available
+		bool bIsORTSupported = (Target.Platform == UnrealTargetPlatform.Win64 /*|| Target.Platform == UnrealTargetPlatform.Linux*/);
+		if (bIsORTSupported)
+		{
+			PublicDefinitions.Add("WITH_FULL_NNI_SUPPORT");
+		}
+
+		ShortName = "DeprecatedNNIWE"; // Could be removed when plugin moves to Experimental, NFL path is too long
 		PCHUsage = ModuleRules.PCHUsageMode.UseExplicitOrSharedPCHs;
 
 		PublicIncludePaths.AddRange(
@@ -29,14 +36,24 @@ public class NeuralNetworkInference : ModuleRules
 		PrivateDependencyModuleNames.AddRange
 			(
 			new string[] {
-                "ONNXRuntime",				// Select this for open-source-based ONNX Runtime
-				//"ONNXRuntimeDLL",			// Select this for DLL-based ONNX Runtime
-				//"ONNXRuntimeDLLHelper",	// Select this for DLL-based ONNX Runtime
 				"Projects",
 				"ThirdPartyHelperAndDLLLoader"
 			}
 		);
 
+		if (bIsORTSupported)
+		{
+
+			PrivateDependencyModuleNames.AddRange
+				(
+				new string[] {
+				"ONNXRuntime"				// Select this for open-source-based ONNX Runtime
+				//"ONNXRuntimeDLL",			// Select this for DLL-based ONNX Runtime (Win64-only)
+				//"ONNXRuntimeDLLHelper"	// Select this for DLL-based ONNX Runtime (Win64-only)
+				}
+			);
+		}
+
 		// Win64-only
 		if (Target.Platform == UnrealTargetPlatform.Win64)
 		{
diff --git a/Engine/Plugins/Experimental/NeuralNetworkInference/Source/NeuralNetworkInference/Private/NeuralNetwork.cpp b/Engine/Plugins/Experimental/NeuralNetworkInference/Source/NeuralNetworkInference/Private/NeuralNetwork.cpp
--- a/Engine/Plugins/Experimental/NeuralNetworkInference/Source/NeuralNetworkInference/Private/NeuralNetwork.cpp
+++ b/Engine/Plugins/Experimental/NeuralNetworkInference/Source/NeuralNetworkInference/Private/NeuralNetwork.cpp
@@ -8,7 +8,6 @@
 #include "Misc/FileHelper.h"
 #include "Misc/Paths.h"
 #include <numeric> // std::accumulate
-#include "RedirectCoutAndCerrToUeLog.h"
 
 //#define WITH_NNI_CPU_NOT_RECOMMENDED // Only for debugging purposes
 
@@ -16,13 +15,17 @@
 NNI_THIRD_PARTY_INCLUDES_START
 #undef check
 #undef TEXT
-#include "onnxruntime/core/session/onnxruntime_cxx_api.h"
-#ifdef PLATFORM_WIN64
-#include "onnxruntime/core/providers/dml/dml_provider_factory.h"
-#endif
-#ifdef WITH_NNI_CPU_NOT_RECOMMENDED
-#include "onnxruntime/core/providers/nni_cpu/nni_cpu_provider_factory.h"
-#endif //WITH_NNI_CPU_NOT_RECOMMENDED
+#ifdef WITH_FULL_NNI_SUPPORT
+	#include "RedirectCoutAndCerrToUeLog.h"
+
+	#include "onnxruntime/core/session/onnxruntime_cxx_api.h"
+	#ifdef PLATFORM_WIN64
+	#include "onnxruntime/core/providers/dml/dml_provider_factory.h"
+	#endif
+	#ifdef WITH_NNI_CPU_NOT_RECOMMENDED
+	#include "onnxruntime/core/providers/nni_cpu/nni_cpu_provider_factory.h"
+	#endif //WITH_NNI_CPU_NOT_RECOMMENDED
+#endif //WITH_FULL_NNI_SUPPORT
 NNI_THIRD_PARTY_INCLUDES_END
 
 
@@ -32,6 +35,7 @@ NNI_THIRD_PARTY_INCLUDES_END
 
 struct UNeuralNetwork::FImpl
 {
+#ifdef WITH_FULL_NNI_SUPPORT
 	TUniquePtr<Ort::Env> Environment;
 	TUniquePtr<Ort::Session> Session;
 	TUniquePtr<Ort::AllocatorWithDefaultOptions> Allocator;
@@ -42,8 +46,10 @@ struct UNeuralNetwork::FImpl
 	bool ConfigureMembers(const ENeuralDeviceType InDeviceType);
 
 	void ConfigureTensors(FNeuralTensors& OutTensors, TArray<bool>* InAreInputTensorSizesVariable = nullptr);
+#endif //WITH_FULL_NNI_SUPPORT
 };
 
+#ifdef WITH_FULL_NNI_SUPPORT
 bool UNeuralNetwork::FImpl::InitializedAndConfigureMembers(TSharedPtr<FImpl>& InOutImpl, const FString& InModelFullFilePath, const ENeuralDeviceType InDeviceType)
 {
 	// Initialize InOutImpl
@@ -215,6 +221,7 @@ void UNeuralNetwork::FImpl::ConfigureTensors(FNeuralTensors& OutTensors, TArray<
 
 	OutTensors.SetFromNetwork(TensorNames, TensorDataTypes, TensorSizes);
 }
+#endif //WITH_FULL_NNI_SUPPORT
 
 
 
@@ -244,6 +251,7 @@ bool UNeuralNetwork::Load(const FString& InModelFilePath)
 	// Clean previous networks
 	bIsLoaded = false;
 
+#ifdef WITH_FULL_NNI_SUPPORT
 	if (InModelFilePath.IsEmpty())
 	{
 		UE_LOG(LogNeuralNetworkInference, Warning, TEXT("UNeuralNetwork::Load(): Input model path was empty."));
@@ -331,14 +339,20 @@ bool UNeuralNetwork::Load(const FString& InModelFilePath)
 //		return Load();
 //	}
 //	// Else
-//	UE_LOG(LogNeuralNetworkInference, Warning, TEXT("Unknown file format \"%s\" used."), *FileExtension);
+//	UE_LOG(LogNeuralNetworkInference, Warning, TEXT("UNeuralNetwork::Load(): Unknown file format \"%s\" used."), *FileExtension);
 //	return false;
+
+#else
+	UE_LOG(LogNeuralNetworkInference, Warning, TEXT("UNeuralNetwork::Load(): Platform or Operating System not suported yet."));
+	return false;
+#endif //WITH_FULL_NNI_SUPPORT
 }
 #endif //WITH_EDITOR
 
 
 bool UNeuralNetwork::Load()
 {
+#ifdef WITH_FULL_NNI_SUPPORT
 	// Clean previous networks
 	bIsLoaded = false;
 
@@ -369,6 +383,11 @@ bool UNeuralNetwork::Load()
 
 	bIsLoaded = true;
 	return bIsLoaded;
+
+#else
+	UE_LOG(LogNeuralNetworkInference, Warning, TEXT("UNeuralNetwork::Load(): Platform or Operating System not suported yet."));
+	return false;
+#endif //WITH_FULL_NNI_SUPPORT
 }
 
 bool UNeuralNetwork::IsLoaded() const
@@ -439,6 +458,7 @@ const FNeuralTensors& UNeuralNetwork::GetOutputTensors() const
 
 void UNeuralNetwork::Run()
 {
+#ifdef WITH_FULL_NNI_SUPPORT
 	// Sanity checks
 	if (!bIsLoaded)
 	{
@@ -457,6 +477,10 @@ void UNeuralNetwork::Run()
 	Impl->Session->Run(Ort::RunOptions{ nullptr },
 		InputTensors.GetTensorNames(), InputTensors.GetONNXRuntimeTensors(), InputTensors.GetNumberTensors(),
 		OutputTensors.GetTensorNames(), OutputTensors.GetONNXRuntimeTensors(), OutputTensors.GetNumberTensors());
+
+#else
+	UE_LOG(LogNeuralNetworkInference, Warning, TEXT("UNeuralNetwork::Run(): Platform or Operating System not suported yet."));
+#endif //WITH_FULL_NNI_SUPPORT
 }
 
 
diff --git a/Engine/Plugins/Experimental/NeuralNetworkInference/Source/NeuralNetworkInference/Private/NeuralTensors.cpp b/Engine/Plugins/Experimental/NeuralNetworkInference/Source/NeuralNetworkInference/Private/NeuralTensors.cpp
--- a/Engine/Plugins/Experimental/NeuralNetworkInference/Source/NeuralNetworkInference/Private/NeuralTensors.cpp
+++ b/Engine/Plugins/Experimental/NeuralNetworkInference/Source/NeuralNetworkInference/Private/NeuralTensors.cpp
@@ -2,19 +2,18 @@
 
 #include "NeuralTensors.h"
 #include "NeuralNetworkInferenceUtils.h"
-#include "RedirectCoutAndCerrToUeLog.h"
 
-#ifdef ONNXRUNTIME_USING_DLL_VERSION
 #include "ThirdPartyWarningDisabler.h"
 NNI_THIRD_PARTY_INCLUDES_START
 #undef check
 #undef TEXT
-#endif
-#include <numeric> // std::accumulate
-#include "onnxruntime/core/session/onnxruntime_cxx_api.h"
-#ifdef ONNXRUNTIME_USING_DLL_VERSION
+#ifdef WITH_FULL_NNI_SUPPORT
+	#include "RedirectCoutAndCerrToUeLog.h"
+
+	#include <numeric> // std::accumulate
+	#include "onnxruntime/core/session/onnxruntime_cxx_api.h"
+#endif //WITH_FULL_NNI_SUPPORT
 NNI_THIRD_PARTY_INCLUDES_END
-#endif
 
 
 
@@ -23,6 +22,7 @@ NNI_THIRD_PARTY_INCLUDES_END
 
 struct FNeuralTensors::FImpl
 {
+#ifdef WITH_FULL_NNI_SUPPORT
 	/** Memory allocator information */
 	Ort::MemoryInfo AllocatorInfo;
 	/** Actual ONNXRuntime tensors */
@@ -34,6 +34,7 @@ struct FNeuralTensors::FImpl
 		: AllocatorInfo(Ort::MemoryInfo::CreateCpu(OrtDeviceAllocator, OrtMemTypeCPU))
 	{
 	}
+#endif //WITH_FULL_NNI_SUPPORT
 };
 
 
@@ -63,6 +64,7 @@ const FNeuralTensor& FNeuralTensors::GetTensor(const int32 InTensorIndex) const
 
 bool FNeuralTensors::Load()
 {
+#ifdef WITH_FULL_NNI_SUPPORT
 #if WITH_EDITOR
 	try
 #endif //WITH_EDITOR
@@ -79,11 +81,18 @@ bool FNeuralTensors::Load()
 		bIsLoaded = false;
 	}
 #endif //WITH_EDITOR
+
+#else
+	bIsLoaded = false;
+	UE_LOG(LogNeuralNetworkInference, Warning, TEXT("FNeuralTensors::Load(): Platform or Operating System not supported yet."));
+#endif //WITH_FULL_NNI_SUPPORT
+
 	return bIsLoaded;
 }
 
 void* FNeuralTensors::GetData(const int32 InTensorIndex)
 {
+#ifdef WITH_FULL_NNI_SUPPORT
 	if (!bIsLoaded && !Load())
 	{
 		UE_LOG(LogNeuralNetworkInference, Warning, TEXT("bIsLoaded is false."));
@@ -114,6 +123,11 @@ void* FNeuralTensors::GetData(const int32 InTensorIndex)
 		return nullptr;
 	}
 #endif //WITH_EDITOR
+
+#else
+	UE_LOG(LogNeuralNetworkInference, Warning, TEXT("FNeuralTensors::GetData(): Platform or Operating System not supported yet."));
+	return nullptr;
+#endif //WITH_FULL_NNI_SUPPORT
 }
 
 int64 FNeuralTensors::GetNumberTensors() const
@@ -123,6 +137,7 @@ int64 FNeuralTensors::GetNumberTensors() const
 
 const void* const FNeuralTensors::GetData(const int32 InTensorIndex) const
 {
+#ifdef WITH_FULL_NNI_SUPPORT
 	if (!bIsLoaded)
 	{
 		UE_LOG(LogNeuralNetworkInference, Warning, TEXT("bIsLoaded is false."));
@@ -153,11 +168,22 @@ const void* const FNeuralTensors::GetData(const int32 InTensorIndex) const
 		return nullptr;
 	}
 #endif //WITH_EDITOR
+
+#else
+	UE_LOG(LogNeuralNetworkInference, Warning, TEXT("FNeuralTensors::GetData(): Platform or Operating System not supported yet."));
+	return nullptr;
+#endif //WITH_FULL_NNI_SUPPORT
 }
 
 FString FNeuralTensors::GetTensorName(const int32 InTensorIndex) const
 {
+#ifdef WITH_FULL_NNI_SUPPORT
 	return FString(ANSI_TO_TCHAR(Impl->TensorNames[InTensorIndex]));
+
+#else
+	UE_LOG(LogNeuralNetworkInference, Warning, TEXT("FNeuralTensors::GetTensorName(): Platform or Operating System not supported yet."));
+	return TEXT("");
+#endif //WITH_FULL_NNI_SUPPORT
 }
 
 const TArray<int64>& FNeuralTensors::GetSizes(const int32 InTensorIndex) const
@@ -193,6 +219,7 @@ void* FNeuralTensors::GetDataPointerMutable(const int32 InTensorIndex)
 
 const char* const* FNeuralTensors::GetTensorNames() const
 {
+#ifdef WITH_FULL_NNI_SUPPORT
 	if (!bIsLoaded)
 	{
 		UE_LOG(LogNeuralNetworkInference, Warning, TEXT("bIsLoaded is false."));
@@ -200,10 +227,16 @@ const char* const* FNeuralTensors::GetTensorNames() const
 	}
 
 	return Impl->TensorNames.GetData();
+
+#else
+	UE_LOG(LogNeuralNetworkInference, Warning, TEXT("FNeuralTensors::GetTensorNames(): Platform or Operating System not supported yet."));
+	return nullptr;
+#endif //WITH_FULL_NNI_SUPPORT
 }
 
 Ort::Value* FNeuralTensors::GetONNXRuntimeTensors()
 {
+#ifdef WITH_FULL_NNI_SUPPORT
 	if (!bIsLoaded && !Load())
 	{
 		UE_LOG(LogNeuralNetworkInference, Warning, TEXT("bIsLoaded is false."));
@@ -211,10 +244,16 @@ Ort::Value* FNeuralTensors::GetONNXRuntimeTensors()
 	}
 
 	return &Impl->OrtTensors[0];
+
+#else
+	UE_LOG(LogNeuralNetworkInference, Warning, TEXT("FNeuralTensors::GetONNXRuntimeTensors(): Platform or Operating System not supported yet."));
+	return nullptr;
+#endif //WITH_FULL_NNI_SUPPORT
 }
 
 const Ort::Value* const FNeuralTensors::GetONNXRuntimeTensors() const
 {
+#ifdef WITH_FULL_NNI_SUPPORT
 	if (!bIsLoaded)
 	{
 		UE_LOG(LogNeuralNetworkInference, Warning, TEXT("bIsLoaded is false."));
@@ -222,6 +261,11 @@ const Ort::Value* const FNeuralTensors::GetONNXRuntimeTensors() const
 	}
 
 	return &Impl->OrtTensors[0];
+
+#else
+	UE_LOG(LogNeuralNetworkInference, Warning, TEXT("FNeuralTensors::GetONNXRuntimeTensors(): Platform or Operating System not supported yet."));
+	return nullptr;
+#endif //WITH_FULL_NNI_SUPPORT
 }
 
 
@@ -231,6 +275,7 @@ const Ort::Value* const FNeuralTensors::GetONNXRuntimeTensors() const
 
 void FNeuralTensors::LinkTensorToONNXRuntime(const int32 InTensorIndex)
 {
+#ifdef WITH_FULL_NNI_SUPPORT
 	if (!bIsLoaded && !Load())
 	{
 		UE_LOG(LogNeuralNetworkInference, Warning, TEXT("bIsLoaded is false."));
@@ -277,11 +322,16 @@ void FNeuralTensors::LinkTensorToONNXRuntime(const int32 InTensorIndex)
 		}
 #endif //WITH_EDITOR
 	}
+
+#else
+	UE_LOG(LogNeuralNetworkInference, Warning, TEXT("FNeuralTensors::LinkTensorToONNXRuntime(): Platform or Operating System not supported yet."));
+#endif //WITH_FULL_NNI_SUPPORT
 }
 
 void FNeuralTensors::SetFromNetwork(TArray<const char*>& InTensorNames, TArray<ENeuralDataType>& InTensorDataTypes,
 		TArray<TArray<int64>>& InSizes)
 {
+#ifdef WITH_FULL_NNI_SUPPORT
 	if (!bIsLoaded && !Load())
 	{
 		UE_LOG(LogNeuralNetworkInference, Warning, TEXT("bIsLoaded is false."));
@@ -314,4 +364,8 @@ void FNeuralTensors::SetFromNetwork(TArray<const char*>& InTensorNames, TArray<E
 		}
 		SetNumUninitialized(InSizes[TensorIndex], InTensorDataTypes[TensorIndex], TensorIndex);
 	}
+
+#else
+	UE_LOG(LogNeuralNetworkInference, Warning, TEXT("FNeuralTensors::SetFromNetwork(): Platform or Operating System not supported yet."));
+#endif //WITH_FULL_NNI_SUPPORT
 }
diff --git a/Engine/Plugins/Experimental/NeuralNetworkInference/Source/NeuralNetworkInferenceEditor/NeuralNetworkInferenceEditor.Build.cs b/Engine/Plugins/Experimental/NeuralNetworkInference/Source/NeuralNetworkInferenceEditor/NeuralNetworkInferenceEditor.Build.cs
--- a/Engine/Plugins/Experimental/NeuralNetworkInference/Source/NeuralNetworkInferenceEditor/NeuralNetworkInferenceEditor.Build.cs
+++ b/Engine/Plugins/Experimental/NeuralNetworkInference/Source/NeuralNetworkInferenceEditor/NeuralNetworkInferenceEditor.Build.cs
@@ -27,20 +27,10 @@ public class NeuralNetworkInferenceEditor : ModuleRules
 		PrivateDependencyModuleNames.AddRange
 			(
 			new string[] {
+				"NeuralNetworkInference",
 				"NeuralNetworkInferenceLegacy",
 				"UnrealEd"
 			}
 		);
-
-		// Win64-only
-		if (Target.Platform == UnrealTargetPlatform.Win64)
-		{
-			PrivateDependencyModuleNames.AddRange
-				(
-				new string[] {
-					"NeuralNetworkInference"
-				}
-			);
-		}
 	}
 }
diff --git a/Engine/Plugins/Experimental/NeuralNetworkInference/Source/NeuralNetworkInferenceEditor/Private/NeuralNetworkAssetTypeActions.cpp b/Engine/Plugins/Experimental/NeuralNetworkInference/Source/NeuralNetworkInferenceEditor/Private/NeuralNetworkAssetTypeActions.cpp
--- a/Engine/Plugins/Experimental/NeuralNetworkInference/Source/NeuralNetworkInferenceEditor/Private/NeuralNetworkAssetTypeActions.cpp
+++ b/Engine/Plugins/Experimental/NeuralNetworkInference/Source/NeuralNetworkInferenceEditor/Private/NeuralNetworkAssetTypeActions.cpp
@@ -1,9 +1,7 @@
 // Copyright Epic Games, Inc. All Rights Reserved.
 
 #include "NeuralNetworkAssetTypeActions.h"
-#if PLATFORM_WINDOWS
 #include "NeuralNetwork.h"
-#endif
 #include "NeuralNetworkInferenceEditorModule.h"
 #include "EditorFramework/AssetImportData.h"
 
@@ -20,11 +18,7 @@ FColor FNeuralNetworkAssetTypeActions::GetTypeColor() const
 
 UClass* FNeuralNetworkAssetTypeActions::GetSupportedClass() const
 {
-#if PLATFORM_WINDOWS
 	return UNeuralNetwork::StaticClass();
-#else
-	return nullptr;
-#endif
 }
 
 bool FNeuralNetworkAssetTypeActions::HasActions(const TArray<UObject*>& InObjects) const
@@ -47,7 +41,6 @@ bool FNeuralNetworkAssetTypeActions::IsImportedAsset() const
 
 void FNeuralNetworkAssetTypeActions::GetResolvedSourceFilePaths(const TArray<UObject*>& TypeAssets, TArray<FString>& OutSourceFilePaths) const
 {
-#if PLATFORM_WINDOWS
 	// Collects the resolved source paths for the imported assets
 	for (UObject* Asset : TypeAssets)
 	{
@@ -61,5 +54,4 @@ void FNeuralNetworkAssetTypeActions::GetResolvedSourceFilePaths(const TArray<UOb
 			OutSourceFilePaths.Add(TEXT(""));
 		}
 	}
-#endif
 }
diff --git a/Engine/Plugins/Experimental/NeuralNetworkInference/Source/NeuralNetworkInferenceEditor/Private/NeuralNetworkFactory.cpp b/Engine/Plugins/Experimental/NeuralNetworkInference/Source/NeuralNetworkInferenceEditor/Private/NeuralNetworkFactory.cpp
--- a/Engine/Plugins/Experimental/NeuralNetworkInference/Source/NeuralNetworkInferenceEditor/Private/NeuralNetworkFactory.cpp
+++ b/Engine/Plugins/Experimental/NeuralNetworkInference/Source/NeuralNetworkInferenceEditor/Private/NeuralNetworkFactory.cpp
@@ -1,9 +1,7 @@
 // Copyright Epic Games, Inc. All Rights Reserved.
 
 #include "NeuralNetworkFactory.h"
-#if PLATFORM_WINDOWS
 #include "NeuralNetwork.h"
-#endif
 #include "NeuralNetworkInferenceEditorUtils.h"
 #include "EditorFramework/AssetImportData.h"
 #include "Misc/Paths.h"
@@ -15,7 +13,6 @@
 
 UNeuralNetworkFactory::UNeuralNetworkFactory()
 {
-#if PLATFORM_WINDOWS
 	SupportedClass = UNeuralNetwork::StaticClass();
 	Formats.Add(TEXT("onnx;ONNX file"));
 	Formats.Add(TEXT("ort;ONNX Runtime (ORT) file"));
@@ -24,7 +21,6 @@ UNeuralNetworkFactory::UNeuralNetworkFactory()
 	bEditorImport = true;
 	bEditAfterNew = true;
 	bText = false;
-#endif
 }
 
 
@@ -34,18 +30,13 @@ UNeuralNetworkFactory::UNeuralNetworkFactory()
 
 UObject* UNeuralNetworkFactory::FactoryCreateNew(UClass* InClass, UObject* InParent, FName InName, EObjectFlags InFlags, UObject* InContext, FFeedbackContext* InWarn)
 {
-#if PLATFORM_WINDOWS
 	// If created with right-click on Content Browser --> Neural Network
 	return NewObject<UNeuralNetwork>(InParent, InName, InFlags);
-#else
-	return nullptr;
-#endif
 }
 
 UObject* UNeuralNetworkFactory::FactoryCreateFile(UClass* InClass, UObject* InParent, FName InName, EObjectFlags InFlags,
 	const FString& InFilename, const TCHAR* Params, FFeedbackContext* Warn, bool& bOutOperationCanceled)
 {
-#if PLATFORM_WINDOWS
 	// If created by dragging a new file into the UE Editor Content Browser
 	if (InFilename.Len() > 0)
 	{
@@ -70,9 +61,7 @@ UObject* UNeuralNetworkFactory::FactoryCreateFile(UClass* InClass, UObject* InPa
 		// If created with right-click on Content Browser --> NeuralNetwork
 		return NewObject<UNeuralNetwork>(InParent, InName, InFlags);
 	}
-#else
 	return nullptr;
-#endif
 }
 
 bool UNeuralNetworkFactory::CanCreateNew() const
@@ -84,20 +73,12 @@ bool UNeuralNetworkFactory::CanCreateNew() const
 
 bool UNeuralNetworkFactory::DoesSupportClass(UClass * Class)
 {
-#if PLATFORM_WINDOWS
 	return (Class == UNeuralNetwork::StaticClass());
-#else
-	return false;
-#endif
 }
 
 UClass* UNeuralNetworkFactory::ResolveSupportedClass()
 {
-#if PLATFORM_WINDOWS
 	return UNeuralNetwork::StaticClass();
-#else
-	return nullptr;
-#endif
 }
 
 bool UNeuralNetworkFactory::FactoryCanImport(const FString& InFilename)
diff --git a/Engine/Plugins/Experimental/NeuralNetworkInference/Source/NeuralNetworkInferenceEditor/Private/NeuralNetworkInferenceEditorModule.cpp b/Engine/Plugins/Experimental/NeuralNetworkInference/Source/NeuralNetworkInferenceEditor/Private/NeuralNetworkInferenceEditorModule.cpp
--- a/Engine/Plugins/Experimental/NeuralNetworkInference/Source/NeuralNetworkInferenceEditor/Private/NeuralNetworkInferenceEditorModule.cpp
+++ b/Engine/Plugins/Experimental/NeuralNetworkInference/Source/NeuralNetworkInferenceEditor/Private/NeuralNetworkInferenceEditorModule.cpp
@@ -1,9 +1,7 @@
 // Copyright Epic Games, Inc. All Rights Reserved.
 
 #include "NeuralNetworkInferenceEditorModule.h"
-#if PLATFORM_WINDOWS
 #include "NeuralNetworkAssetTypeActions.h"
-#endif
 #include "NeuralNetworkLegacyAssetTypeActions.h"
 #include "IAssetTypeActions.h"
 #include "Interfaces/IPluginManager.h"
@@ -38,7 +36,7 @@ void FNeuralNetworkInferenceEditorModule::StartupModule()
 {
 	// UNeuralNetwork - Register asset types
 	IAssetTools& AssetTools = FModuleManager::LoadModuleChecked<FAssetToolsModule>("AssetTools").Get();
-#if PLATFORM_WINDOWS
+#ifdef WITH_FULL_NNI_SUPPORT
 	Actions.Emplace(MakeShared<FNeuralNetworkAssetTypeActions>());
 	AssetTools.RegisterAssetTypeActions(Actions.Last().ToSharedRef());
 #endif
diff --git a/Engine/Plugins/Experimental/NeuralNetworkInference/Source/NeuralNetworkInferenceEditor/Private/NeuralNetworkReimportFactory.cpp b/Engine/Plugins/Experimental/NeuralNetworkInference/Source/NeuralNetworkInferenceEditor/Private/NeuralNetworkReimportFactory.cpp
--- a/Engine/Plugins/Experimental/NeuralNetworkInference/Source/NeuralNetworkInferenceEditor/Private/NeuralNetworkReimportFactory.cpp
+++ b/Engine/Plugins/Experimental/NeuralNetworkInference/Source/NeuralNetworkInferenceEditor/Private/NeuralNetworkReimportFactory.cpp
@@ -1,9 +1,7 @@
 // Copyright Epic Games, Inc. All Rights Reserved.
 
 #include "NeuralNetworkReimportFactory.h"
-#if PLATFORM_WINDOWS
 #include "NeuralNetwork.h"
-#endif
 #include "NeuralNetworkInferenceEditorUtils.h"
 #include "EditorFramework/AssetImportData.h"
 #include "HAL/FileManager.h"
@@ -15,7 +13,6 @@
 
 UNeuralNetworkReimportFactory::UNeuralNetworkReimportFactory()
 {
-#if PLATFORM_WINDOWS
 	SupportedClass = UNeuralNetwork::StaticClass();
 	Formats.Add(TEXT("onnx;ONNX file"));
 	Formats.Add(TEXT("ort;ONNX Runtime (ORT) file"));
@@ -25,7 +22,6 @@ UNeuralNetworkReimportFactory::UNeuralNetworkReimportFactory()
 
 	// Required to allow other re importers to do their CanReimport checks first, and if they fail this re importer will catch it.
 	ImportPriority = DefaultImportPriority - 1;
-#endif
 }
 
 
@@ -46,7 +42,6 @@ bool UNeuralNetworkReimportFactory::FactoryCanImport(const FString& Filename)
 
 bool UNeuralNetworkReimportFactory::CanReimport(UObject* Obj, TArray<FString>& OutFilenames)
 {
-#if PLATFORM_WINDOWS
 	UNeuralNetwork* Network = Cast<UNeuralNetwork>(Obj);
 	if (Network)
 	{
@@ -66,13 +61,11 @@ bool UNeuralNetworkReimportFactory::CanReimport(UObject* Obj, TArray<FString>& O
 		}
 		return true;
 	}
-#endif
 	return false;
 }
 
 void UNeuralNetworkReimportFactory::SetReimportPaths(UObject* Obj, const TArray<FString>& NewReimportPaths)
 {
-#if PLATFORM_WINDOWS
 	UNeuralNetwork* Network = Cast<UNeuralNetwork>(Obj);
 	if (Network && ensure(NewReimportPaths.Num() == 1))
 	{
@@ -86,12 +79,10 @@ void UNeuralNetworkReimportFactory::SetReimportPaths(UObject* Obj, const TArray<
 			UE_LOG(LogNeuralNetworkInferenceEditor, Warning, TEXT("UNeuralNetworkReimportFactory::Reimport(): AssetImportData was nullptr."));
 		}
 	}
-#endif
 }
 
 EReimportResult::Type UNeuralNetworkReimportFactory::Reimport(UObject* Obj)
 {
-#if PLATFORM_WINDOWS
 	UNeuralNetwork* Network = Cast<UNeuralNetwork>(Obj);
 	if (!Network)
 	{
@@ -128,7 +119,6 @@ EReimportResult::Type UNeuralNetworkReimportFactory::Reimport(UObject* Obj)
 		AssetImportData->Update(ImportedFilename);
 		return EReimportResult::Succeeded;
 	}
-#endif
 	UE_LOG(LogNeuralNetworkInferenceEditor, Warning, TEXT("UNeuralNetworkReimportFactory::Reimport(): Load failed."));
 	return EReimportResult::Failed;
 }
diff --git a/Engine/Plugins/Experimental/NeuralNetworkInference/Source/NeuralNetworkInferenceQA/NeuralNetworkInferenceQA.Build.cs b/Engine/Plugins/Experimental/NeuralNetworkInference/Source/NeuralNetworkInferenceQA/NeuralNetworkInferenceQA.Build.cs
--- a/Engine/Plugins/Experimental/NeuralNetworkInference/Source/NeuralNetworkInferenceQA/NeuralNetworkInferenceQA.Build.cs
+++ b/Engine/Plugins/Experimental/NeuralNetworkInference/Source/NeuralNetworkInferenceQA/NeuralNetworkInferenceQA.Build.cs
@@ -29,6 +29,7 @@ public class NeuralNetworkInferenceQA : ModuleRules
 			(
 			new string[] {
 				"ModelProto",
+				"NeuralNetworkInference",
 				"NeuralNetworkInferenceCore",
 				"NeuralNetworkInferenceLegacy"
 			}
@@ -37,13 +38,6 @@ public class NeuralNetworkInferenceQA : ModuleRules
 		// Win64-only
 		if (Target.Platform == UnrealTargetPlatform.Win64)
 		{
-			PrivateDependencyModuleNames.AddRange
-				(
-				new string[] {
-					"NeuralNetworkInference"
-				}
-			);
-
 			// Editor-only (to be deprecated)
 			if (Target.bBuildEditor)
 			{
diff --git a/Engine/Plugins/Experimental/NeuralNetworkInference/Source/NeuralNetworkInferenceQA/Private/ModelUnitTester.cpp b/Engine/Plugins/Experimental/NeuralNetworkInference/Source/NeuralNetworkInferenceQA/Private/ModelUnitTester.cpp
--- a/Engine/Plugins/Experimental/NeuralNetworkInference/Source/NeuralNetworkInferenceQA/Private/ModelUnitTester.cpp
+++ b/Engine/Plugins/Experimental/NeuralNetworkInference/Source/NeuralNetworkInferenceQA/Private/ModelUnitTester.cpp
@@ -12,7 +12,7 @@
 
 void FModelUnitTester::GlobalTest(const FString& InModelsDirectory)
 {
-#ifdef PLATFORM_WIN64
+#ifdef WITH_FULL_NNI_SUPPORT
 	// Model load, accuracy, and speed test
 	const TArray<FString> ModelNames({ TEXT("MLRigDeformer"), TEXT("cloth_network"), TEXT("hyprsense"), TEXT("rocket_league") });
 	const TArray<float> InputArrayValues({ 1.f, 0.f, -1.f, 100.f, -100.f, 0.5f, -0.5f }); // This one can be shorter than CPU/GPUGroundTruths
@@ -24,8 +24,8 @@ void FModelUnitTester::GlobalTest(const FString& InModelsDirectory)
 	const TArray<int32> GPURepetitions({ 1000, 1000, 100, 1000 });
 	ModelLoadAccuracyAndSpeedTests(InModelsDirectory, ModelNames, InputArrayValues, CPUGroundTruths, GPUGroundTruths, CPURepetitions, GPURepetitions);
 #else
-	UE_LOG(LogNeuralNetworkInferenceQA, Display, TEXT("--------------- FModelUnitTester test skipped (only in Windows)."));
-#endif //PLATFORM_WIN64
+	UE_LOG(LogNeuralNetworkInferenceQA, Display, TEXT("--------------- FModelUnitTester test skipped (only if WITH_FULL_NNI_SUPPORT)."));
+#endif //WITH_FULL_NNI_SUPPORT
 }
 
 
@@ -33,8 +33,6 @@ void FModelUnitTester::GlobalTest(const FString& InModelsDirectory)
 /* FModelUnitTester static private functions
  *****************************************************************************/
 
-#ifdef PLATFORM_WIN64
-
 void FModelUnitTester::ModelLoadAccuracyAndSpeedTests(const FString& InModelsDirectory, const TArray<FString>& InModelNames, const TArray<float>& InInputArrayValues,
 	const TArray<TArray<double>>& InCPUGroundTruths, const TArray<TArray<double>>& InGPUGroundTruths, const TArray<int32>& InCPURepetitions, const TArray<int32>& InGPURepetitions)
 {
@@ -312,5 +310,3 @@ double FModelUnitTester::GetAveragedL1NormDiff(const TArray<float>& InArray1, co
 	AveragedL1NormDiff /= InArray1.Num();
 	return AveragedL1NormDiff;
 }
-
-#endif //PLATFORM_WIN64
diff --git a/Engine/Plugins/Experimental/NeuralNetworkInference/Source/NeuralNetworkInferenceQA/Private/ModelUnitTester.h b/Engine/Plugins/Experimental/NeuralNetworkInference/Source/NeuralNetworkInferenceQA/Private/ModelUnitTester.h
--- a/Engine/Plugins/Experimental/NeuralNetworkInference/Source/NeuralNetworkInferenceQA/Private/ModelUnitTester.h
+++ b/Engine/Plugins/Experimental/NeuralNetworkInference/Source/NeuralNetworkInferenceQA/Private/ModelUnitTester.h
@@ -3,16 +3,13 @@
 #pragma once
 
 #include "CoreMinimal.h"
-#ifdef PLATFORM_WIN64
 #include "NeuralNetwork.h"
-#endif //PLATFORM_WIN64
 
 class FModelUnitTester
 {
 public:
 	static void GlobalTest(const FString& InModelsDirectory);
 
-#ifdef PLATFORM_WIN64
 private:
 	/**
 	 * It runs a full model test on the desired model.
@@ -34,5 +31,4 @@ private:
 	static void ModelSpeedTest(const FString& InUAssetPath, const int32 InCPURepetitions, const int32 InGPURepetitions);
 	static double GetAveragedL1Norm(const TArray<float>& InArray);
 	static double GetAveragedL1NormDiff(const TArray<float>& InArray1, const TArray<float>& InArray2);
-#endif //PLATFORM_WIN64
 };
