NeuralNetworkInference (NNI) plugin:
- NNI + NNIBackEnd = NNI.

#rnx

diff --git a/Engine/Plugins/Experimental/NeuralNetworkInference/NeuralNetworkInference.uplugin b/Engine/Plugins/Experimental/NeuralNetworkInference/NeuralNetworkInference.uplugin
--- a/Engine/Plugins/Experimental/NeuralNetworkInference/NeuralNetworkInference.uplugin
+++ b/Engine/Plugins/Experimental/NeuralNetworkInference/NeuralNetworkInference.uplugin
@@ -20,11 +20,6 @@
 			"Type": "RuntimeAndProgram",
 			"LoadingPhase": "Default"
 		},
-		{
-			"Name": "NeuralNetworkInferenceBackEnd",
-			"Type": "RuntimeAndProgram",
-			"LoadingPhase": "Default"
-		},
 		{
 			"Name": "NeuralNetworkInferenceCore",
 			"Type": "RuntimeAndProgram",
diff --git a/Engine/Plugins/Experimental/NeuralNetworkInference/Source/Internal/NeuralNetworkInferenceLegacy/NeuralNetworkInferenceLegacy.Build.cs b/Engine/Plugins/Experimental/NeuralNetworkInference/Source/Internal/NeuralNetworkInferenceLegacy/NeuralNetworkInferenceLegacy.Build.cs
--- a/Engine/Plugins/Experimental/NeuralNetworkInference/Source/Internal/NeuralNetworkInferenceLegacy/NeuralNetworkInferenceLegacy.Build.cs
+++ b/Engine/Plugins/Experimental/NeuralNetworkInference/Source/Internal/NeuralNetworkInferenceLegacy/NeuralNetworkInferenceLegacy.Build.cs
@@ -22,7 +22,7 @@ public class NeuralNetworkInferenceLegacy : ModuleRules
 				"CoreUObject",
 				"Engine",
 				"ModelProto",
-				"NeuralNetworkInferenceBackEnd",
+				"NeuralNetworkInference",
 				"NeuralNetworkInferenceCore"
 			}
 		);
diff --git a/Engine/Plugins/Experimental/NeuralNetworkInference/Source/NeuralNetworkInference/NeuralNetworkInference.Build.cs b/Engine/Plugins/Experimental/NeuralNetworkInference/Source/NeuralNetworkInference/NeuralNetworkInference.Build.cs
--- a/Engine/Plugins/Experimental/NeuralNetworkInference/Source/NeuralNetworkInference/NeuralNetworkInference.Build.cs
+++ b/Engine/Plugins/Experimental/NeuralNetworkInference/Source/NeuralNetworkInference/NeuralNetworkInference.Build.cs
@@ -13,7 +13,7 @@ public class NeuralNetworkInference : ModuleRules
 			PublicDefinitions.Add("WITH_FULL_NNI_SUPPORT");
 		}
 
-		ShortName = "DeprecatedNNIWE"; // Could be removed when plugin moves to Experimental, NFL path is too long
+		ShortName = "NNI"; // Could be removed when plugin moves to Experimental, NFL path is too long
 		PCHUsage = ModuleRules.PCHUsageMode.UseExplicitOrSharedPCHs;
 
 		PublicIncludePaths.AddRange(
@@ -28,7 +28,6 @@ public class NeuralNetworkInference : ModuleRules
 				"Core",
 				"CoreUObject",
 				"Engine",
-				"NeuralNetworkInferenceBackEnd",
 				"NeuralNetworkInferenceCore"
 			}
 		);
@@ -36,6 +35,12 @@ public class NeuralNetworkInference : ModuleRules
 		PrivateDependencyModuleNames.AddRange
 			(
 			new string[] {
+				// Backend-related
+				"ModelProto",
+				"NeuralNetworkInferenceShaders",
+				"RenderCore",
+				"RHI",
+				// ORT-related
 				"Projects",
 				"ThirdPartyHelperAndDLLLoader"
 			}
diff --git a/Engine/Plugins/Experimental/NeuralNetworkInference/Source/NeuralNetworkInference/Private/NeuralInt64ArrayUInt32Buffer.cpp b/Engine/Plugins/Experimental/NeuralNetworkInference/Source/NeuralNetworkInference/Private/NeuralInt64ArrayUInt32Buffer.cpp
new file mode 100644
--- /dev/null
+++ b/Engine/Plugins/Experimental/NeuralNetworkInference/Source/NeuralNetworkInference/Private/NeuralInt64ArrayUInt32Buffer.cpp
@@ -0,0 +1,92 @@
+// Copyright Epic Games, Inc. All Rights Reserved.
+
+#include "NeuralInt64ArrayUInt32Buffer.h"
+#include "NeuralNetworkInferenceUtils.h"
+#include "NeuralNetworkInferenceUtilsGPU.h"
+
+
+
+/* FNeuralInt64ArrayUInt32Buffer structors
+ *****************************************************************************/
+
+FNeuralInt64ArrayUInt32Buffer::FNeuralInt64ArrayUInt32Buffer(const TArray<int64>& InInt64Array)
+{
+	Copy(InInt64Array);
+}
+
+FNeuralInt64ArrayUInt32Buffer::FNeuralInt64ArrayUInt32Buffer(const TArray<uint32>& InUInt32Array, const bool bShouldCreate64BitVersion)
+{
+	Copy(InUInt32Array, bShouldCreate64BitVersion);
+}
+
+
+
+/* FNeuralInt64ArrayUInt32Buffer public functions
+ *****************************************************************************/
+
+const TArray<int64>& FNeuralInt64ArrayUInt32Buffer::GetInt64() const
+{
+	if (Int64Array.IsEmpty())
+	{
+		UE_LOG(LogNeuralNetworkInference, Warning,
+			TEXT("FNeuralInt64ArrayUInt32Buffer::GetInt64(): Int64Array is empty. Move() might have been never called or called with bShouldCreate64BitVersion = false."));
+	}
+	return Int64Array;
+}
+
+void FNeuralInt64ArrayUInt32Buffer::Init(const uint32 InArrayValue, const uint32 InArraySize, const bool bShouldCreate64BitVersion)
+{
+	TArray<uint32> UInt32ArrayTemp;
+	UInt32ArrayTemp.Init(InArrayValue, InArraySize);
+	Move(UInt32ArrayTemp, bShouldCreate64BitVersion);
+}
+
+void FNeuralInt64ArrayUInt32Buffer::Copy(const TArray<int64>& InInt64Array)
+{
+	TArray<int64> Int64ArrayTemp = InInt64Array;
+	Move(Int64ArrayTemp);
+}
+
+void FNeuralInt64ArrayUInt32Buffer::Copy(const TArray<uint32>& InUInt32Array, const bool bShouldCreate64BitVersion)
+{
+	TArray<uint32> UInt32ArrayTemp = InUInt32Array;
+	Move(UInt32ArrayTemp, bShouldCreate64BitVersion);
+}
+
+void FNeuralInt64ArrayUInt32Buffer::Move(TArray<int64>& InInt64Array)
+{
+	// Int64Array
+	Swap(Int64Array, InInt64Array);
+	// UInt32Array
+	UInt32Array.SetNumUninitialized(Int64Array.Num());
+	for (int32 Index = 0; Index < UInt32Array.Num(); ++Index)
+	{
+		UInt32Array[Index] = (uint32)Int64Array[Index];
+	}
+}
+
+void FNeuralInt64ArrayUInt32Buffer::Move(TArray<uint32>& InUInt32Array, const bool bShouldCreate64BitVersion)
+{
+	// UInt32Array
+	Swap(UInt32Array, InUInt32Array);
+	// Int64Array
+	if (bShouldCreate64BitVersion)
+	{
+		Int64Array.SetNumUninitialized(UInt32Array.Num());
+		for (int32 Index = 0; Index < Int64Array.Num(); ++Index)
+		{
+			Int64Array[Index] = (int64)UInt32Array[Index];
+		}
+	}
+}
+
+void FNeuralInt64ArrayUInt32Buffer::CreateAndLoadSRVBuffer(const TCHAR* InDebugName)
+{
+	// Sanity check
+	if (UInt32ReadBuffer.IsValid())
+	{
+		UE_LOG(LogNeuralNetworkInference, Warning, TEXT("FNeuralInt64ArrayUInt32Buffer::CreateAndLoadSRVBuffer(): UInt32ReadBuffer was not nullptr."));
+	}
+	// Create GPU memory for it
+	FNeuralNetworkInferenceUtilsGPU::CreateAndLoadSRVBuffer(UInt32ReadBuffer, UInt32Array, InDebugName);
+}
diff --git a/Engine/Plugins/Experimental/NeuralNetworkInference/Source/NeuralNetworkInference/Private/NeuralNetworkInferenceUtils.cpp b/Engine/Plugins/Experimental/NeuralNetworkInference/Source/NeuralNetworkInference/Private/NeuralNetworkInferenceUtils.cpp
new file mode 100644
--- /dev/null
+++ b/Engine/Plugins/Experimental/NeuralNetworkInference/Source/NeuralNetworkInference/Private/NeuralNetworkInferenceUtils.cpp
@@ -0,0 +1,57 @@
+// Copyright Epic Games, Inc. All Rights Reserved.
+
+#include "NeuralNetworkInferenceUtils.h"
+#include "HAL/PlatformTime.h"
+
+
+
+/* FNeuralNetworkInferenceUtils public functions
+ *****************************************************************************/
+
+bool FNeuralNetworkInferenceUtils::SizeSanityChecks(const TArray<FNeuralTensor*>& InTensorArray, const int32 InMinNum, const int32 InMaxNum,
+	const int32 InMinDimensions, const int32 InMaxDimensions, const int32 InDimensionRangeFirst, const int32 InDimensionRangeLast)
+{
+	// Do InMinNum/InMaxNum checks
+	if (InTensorArray.Num() < InMinNum || InTensorArray.Num() > InMaxNum)
+	{
+		UE_LOG(LogNeuralNetworkInference, Warning,
+			TEXT("FNeuralNetworkInferenceUtils::SizeSanityChecks(): %d <= InputTensors.Num() (%d) <= %d failed."), InMinNum, InTensorArray.Num(), InMaxNum);
+		return false;
+	}
+	// Do Tensor != nullptr checks
+	for (const FNeuralTensor* const Tensor : InTensorArray)
+	{
+		if (!Tensor)
+		{
+			UE_LOG(LogNeuralNetworkInference, Warning,
+				TEXT("FNeuralNetworkInferenceUtils::SizeSanityChecks(): InTensorArray contained nullptr elements."));
+			return false;
+		}
+	}
+	// Do InMinDimensions/InMaxDimensions checks (if enabled)
+	if (InMinDimensions > -1 || InDimensionRangeLast > -1)
+	{
+		const int32 DimensionRangeFirst = (InDimensionRangeFirst > -1 ? InDimensionRangeFirst : 0);
+		const int32 DimensionRangeLast = FMath::Min(InTensorArray.Num(), // Make sure it does not go out of indexes (i.e., InMinNum != InMaxNum and InTensorArray < InMaxNum)
+			InDimensionRangeLast > -1 ? InDimensionRangeLast + 1 : InTensorArray.Num()); // From 1-index to 0-index
+		for (int32 TensorIndex = DimensionRangeFirst; TensorIndex < DimensionRangeLast; ++TensorIndex)
+		{
+			const FNeuralTensor* const Tensor = InTensorArray[TensorIndex];
+			if (InMinDimensions > -1 && Tensor->GetNumberDimensions() < InMinDimensions)
+			{
+				UE_LOG(LogNeuralNetworkInference, Warning,
+					TEXT("FNeuralNetworkInferenceUtils::SizeSanityChecks(): %d <= Tensors[%d]->GetNumberDimensions() (%d) failed."),
+					InMinDimensions, TensorIndex, Tensor->GetNumberDimensions());
+				return false;
+			}
+			else if (InMaxDimensions > -1 && Tensor->GetNumberDimensions() > InMaxDimensions)
+			{
+				UE_LOG(LogNeuralNetworkInference, Warning,
+					TEXT("FNeuralNetworkInferenceUtils::SizeSanityChecks(): Tensors[%d]->GetNumberDimensions() (%d) <= %d failed."),
+					TensorIndex, Tensor->GetNumberDimensions(), InMaxDimensions);
+				return false;
+			}
+		}
+	}
+	return true;
+}
diff --git a/Engine/Plugins/Experimental/NeuralNetworkInference/Source/NeuralNetworkInference/Private/NeuralNetworkInferenceUtils.h b/Engine/Plugins/Experimental/NeuralNetworkInference/Source/NeuralNetworkInference/Private/NeuralNetworkInferenceUtils.h
--- a/Engine/Plugins/Experimental/NeuralNetworkInference/Source/NeuralNetworkInference/Private/NeuralNetworkInferenceUtils.h
+++ b/Engine/Plugins/Experimental/NeuralNetworkInference/Source/NeuralNetworkInference/Private/NeuralNetworkInferenceUtils.h
@@ -3,5 +3,44 @@
 #pragma once
 
 #include "CoreMinimal.h"
+#include "NeuralTensor.h"
+#include "Algo/Accumulate.h"
 
 DEFINE_LOG_CATEGORY_STATIC(LogNeuralNetworkInference, Display, All);
+
+class FNeuralNetworkInferenceUtils
+{
+public:
+	/**
+	 * SizeSanityChecks() makes sure the tensor has the expected size (either InNum or between InMinNum and InMaxNum) and that all elements are not nullptr.
+	 * @param InMinNum/InMaxNum: Range of elements that InTensorArray.Num() must satisfy.
+	 * @param InMinDimensions/InMaxDimensions: If not -1, range of dimensions that each element in InTensorArray.Num() must satisfy.
+	 * @param InMinDimensionsFirst/InMaxDimensionsLast: If not -1, element range for which InMinDimensions/InMaxDimensions is applied to.
+	 */
+	FORCEINLINE static bool SizeSanityChecks(const TArray<FNeuralTensor*>& InTensorArray, const int32 InNum);
+	static bool SizeSanityChecks(const TArray<FNeuralTensor*>& InTensorArray, const int32 InMinNum, const int32 InMaxNum,
+		const int32 InMinDimensions = -1, const int32 InMaxDimensions = -1, const int32 InDimensionRangeFirst = -1, const int32 InDimensionRangeLast = -1);
+
+	template <typename T>
+	struct MultiplyOperation
+	{
+		constexpr T operator()(const T& InLeft, const T& InRight) const { return InLeft * InRight; }
+	};
+
+	template <typename ValueType, typename ContainerType>
+	static ValueType Product(const ContainerType& InContainer);
+};
+
+
+
+// Inline implementations
+bool FNeuralNetworkInferenceUtils::SizeSanityChecks(const TArray<FNeuralTensor*>& InTensorArray, const int32 InNum)
+{
+	return SizeSanityChecks(InTensorArray, InNum, InNum);
+}
+
+template <typename ValueType, typename ContainerType>
+ValueType FNeuralNetworkInferenceUtils::Product(const ContainerType& InContainer)
+{
+	return Algo::Accumulate(InContainer, ValueType(1), MultiplyOperation<ValueType>());
+}
diff --git a/Engine/Plugins/Experimental/NeuralNetworkInference/Source/NeuralNetworkInference/Private/NeuralNetworkInferenceUtilsGPU.cpp b/Engine/Plugins/Experimental/NeuralNetworkInference/Source/NeuralNetworkInference/Private/NeuralNetworkInferenceUtilsGPU.cpp
new file mode 100644
--- /dev/null
+++ b/Engine/Plugins/Experimental/NeuralNetworkInference/Source/NeuralNetworkInference/Private/NeuralNetworkInferenceUtilsGPU.cpp
@@ -0,0 +1,48 @@
+// Copyright Epic Games, Inc. All Rights Reserved.
+
+#include "NeuralNetworkInferenceUtilsGPU.h"
+#include "NeuralNetworkInferenceUtils.h"
+#include "NeuralTensorResourceArray.h"
+
+
+
+/* FNeuralNetworkInferenceUtilsGPU public functions
+ *****************************************************************************/
+
+void FNeuralNetworkInferenceUtilsGPU::CreateAndLoadSRVBuffer(TSharedPtr<FReadBuffer>& OutReadBuffer, const TArray<uint32>& InArrayData, const TCHAR* InDebugName)
+{
+	if (OutReadBuffer.IsValid())
+	{
+		OutReadBuffer->Release();
+	}
+	OutReadBuffer = MakeShared<FReadBuffer>();
+	TSharedPtr<FNeuralTensorResourceArray> TensorResourceArray = MakeShared<FNeuralTensorResourceArray>((void*)InArrayData.GetData(), sizeof(uint32) * InArrayData.Num());
+	OutReadBuffer->Initialize(InDebugName, sizeof(uint32), InArrayData.Num(), PF_R32_UINT, BUF_ShaderResource | BUF_Static, TensorResourceArray.Get());
+}
+
+bool FNeuralNetworkInferenceUtilsGPU::GPUSanityChecks(const FRDGBuilder* const InGraphBuilder)
+{
+	// Sanity checks
+	if (!IsInRenderingThread())
+	{
+		UE_LOG(LogNeuralNetworkInference, Warning, TEXT("GPUSanityChecks(): IsInRenderingThread() should be true."));
+		return false;
+	}
+	else if (!InGraphBuilder)
+	{
+		UE_LOG(LogNeuralNetworkInference, Warning, TEXT("GPUSanityChecks(): InOutGraphBuilder cannot be nullptr."));
+		return false;
+	}
+	return true;
+}
+
+bool FNeuralNetworkInferenceUtilsGPU::GPUSanityChecks(const FRDGBuilder* const InGraphBuilder, const bool bInIsLoaded)
+{
+	// Sanity checks
+	if (!bInIsLoaded)
+	{
+		UE_LOG(LogNeuralNetworkInference, Warning, TEXT("GPUSanityChecks(): bIsLoaded should be true."));
+		return false;
+	}
+	return GPUSanityChecks(InGraphBuilder);
+}
diff --git a/Engine/Plugins/Experimental/NeuralNetworkInference/Source/NeuralNetworkInference/Private/NeuralNetworkInferenceUtilsGPU.h b/Engine/Plugins/Experimental/NeuralNetworkInference/Source/NeuralNetworkInference/Private/NeuralNetworkInferenceUtilsGPU.h
new file mode 100644
--- /dev/null
+++ b/Engine/Plugins/Experimental/NeuralNetworkInference/Source/NeuralNetworkInference/Private/NeuralNetworkInferenceUtilsGPU.h
@@ -0,0 +1,23 @@
+// Copyright Epic Games, Inc. All Rights Reserved.
+
+#pragma once
+
+#include "CoreMinimal.h"
+#include "RenderGraphBuilder.h"
+#include "RHI.h"
+
+class FNeuralNetworkInferenceUtilsGPU
+{
+public:
+	/**
+	 * If OutReadBufferis nullptr, it creates a new FReadBuffer pointer into OutReadBuffer and copies the data from InArrayData.
+	 * @param InDebugName Input name for FReadBuffer::Initialize().
+	 */
+	static void CreateAndLoadSRVBuffer(TSharedPtr<FReadBuffer>& OutReadBuffer, const TArray<uint32>& InArrayData, const TCHAR* InDebugName);
+
+	/**
+	 * Sanity checks when running the forward operators or their related GPU functions.
+	 */
+	static bool GPUSanityChecks(const FRDGBuilder* const InGraphBuilder);
+	static bool GPUSanityChecks(const FRDGBuilder* const InGraphBuilder, const bool bInIsLoaded);
+};
diff --git a/Engine/Plugins/Experimental/NeuralNetworkInference/Source/NeuralNetworkInference/Private/NeuralOperator.cpp b/Engine/Plugins/Experimental/NeuralNetworkInference/Source/NeuralNetworkInference/Private/NeuralOperator.cpp
new file mode 100644
--- /dev/null
+++ b/Engine/Plugins/Experimental/NeuralNetworkInference/Source/NeuralNetworkInference/Private/NeuralOperator.cpp
@@ -0,0 +1,168 @@
+// Copyright Epic Games, Inc. All Rights Reserved.
+
+#include "NeuralOperator.h"
+#include "NeuralNetworkInferenceUtils.h"
+
+
+
+/* FPrivateNeuralOperator
+ *****************************************************************************/
+
+class FPrivateNeuralOperator
+{
+public:
+	static void SanityCheckInlining(const bool bInShouldBeInlined, const int32 InInlinedTensor, const FString& InName);
+};
+
+void FPrivateNeuralOperator::SanityCheckInlining(const bool bInShouldBeInlined, const int32 InInlinedTensor, const FString& InName)
+{
+	checkf((bInShouldBeInlined && InInlinedTensor >= 0) || (!bInShouldBeInlined && InInlinedTensor < 0),
+		TEXT("FNeuralOperator-%s::SanityCheckInlining(): bInShouldBeInlined (%d) vs. InInlinedTensor (%d) mismatch. If"
+			" bInShouldBeInlined == true, InInlinedTensor should be >=0. If false, InInlinedTensor should be < 0. Possible reasons:"
+			" 1) Non-const input called by a non-inlined operator;"
+			" 2) Non-existing output called by an inlined operator;"
+			" 3) SetDataType() was not called."),
+		*InName, bInShouldBeInlined, InInlinedTensor, (int32)bInShouldBeInlined);
+}
+
+
+
+/* FNeuralOperator structors
+ *****************************************************************************/
+
+FNeuralOperator::FNeuralOperator(const FString& InName, const int32 InVersion, const int32 InInlinedTensor)
+	: bIsLoaded(false)
+	, Name(InName)
+	, Version(InVersion)
+	, InlinedTensor(InInlinedTensor)
+{
+}
+
+FNeuralOperator::FNeuralOperator(const FString& InName, const int32 InVersion, const TSet<uint32>& InPotentialInlinedTensors)
+	: bIsLoaded(false)
+	, Name(InName)
+	, Version(InVersion)
+	, PotentiallyInlinedTensors(InPotentialInlinedTensors)
+	, InlinedTensor(-1)
+{
+}
+
+FNeuralOperator::~FNeuralOperator()
+{
+}
+
+
+
+/* FNeuralOperator public functions
+ *****************************************************************************/
+
+int32 FNeuralOperator::SetInputTensorsAndGetInlinedIndex(const TArray<FNeuralTensor*>& InInputTensors)
+{
+	PrivateInputTensors = InInputTensors;
+	if (!EstimateInlinedTensorFromPotentialOnes())
+	{
+		UE_LOG(LogNeuralNetworkInference, Warning, TEXT("FNeuralOperator::SetInputTensorsAndGetInlinedIndex(): EstimateInlinedTensorFromPotentialOnes() returned false."));
+		InlinedTensor = -2;
+	}
+	return InlinedTensor;
+}
+
+void FNeuralOperator::SetAuxiliaryTensors(const TArray<FNeuralTensor*>& InAuxiliaryTensors)
+{
+	checkf(GetMaximumPossibleNumberAuxiliaryTensorsPerOperator() >= InAuxiliaryTensors.Num(),
+		TEXT("FNeuralOperator-%s::SetAuxiliaryTensors(): Either 1) This is a new operator so GetMaximumPossibleNumberAuxiliaryTensorsPerOperator() must be modified to return the new maximum (%d) or 2) InAuxiliaryTensors.Num() is too big for this operator."),
+		*Name, InAuxiliaryTensors.Num());
+	AuxiliaryTensors = InAuxiliaryTensors;
+}
+
+bool FNeuralOperator::Configure()
+{
+	// Trick to only run until one of them is false (or otherwise it runs them all and return true)
+	return (ConfigureOutputAndInternalVariablesAndSanityChecks() && SetWhetherInputTensorsNeedTransferToGPUForGPUMode() && SetWhetherOutputTensorsNeedTransferToGPUForGPUMode() && SetWhetherAuxiliaryTensorsNeedTransferToGPUForGPUMode());
+}
+
+FString FNeuralOperator::ToString() const
+{
+	const FString InlinedText = (InlinedTensor < 0 ? TEXT("not inlined") : TEXT("inlined on tensor ") + FString::FromInt(InlinedTensor) );
+	return FString::Format(TEXT("{0}_v{1} ({2})"), { Name, FString::FromInt(Version), InlinedText });
+}
+
+
+
+/* FNeuralOperator protected functions
+ *****************************************************************************/
+
+bool FNeuralOperator::SetWhetherInputTensorsNeedTransferToGPUForGPUMode()
+{
+	// Enable GPU mode on input tensors
+	for (FNeuralTensor* PrivateInputTensor : PrivateInputTensors)
+	{
+		PrivateInputTensor->SetEnableGPU(true);
+	}
+	return true;
+}
+
+bool FNeuralOperator::SetWhetherOutputTensorsNeedTransferToGPUForGPUMode()
+{
+	// Enable GPU mode on output tensors
+	for (FNeuralTensor* PrivateOutputTensor : PrivateOutputTensors)
+	{
+		PrivateOutputTensor->SetEnableGPU(true);
+	}
+	return true;
+}
+
+bool FNeuralOperator::SetWhetherAuxiliaryTensorsNeedTransferToGPUForGPUMode()
+{
+	// Enable GPU mode on auxiliary tensors
+	for (FNeuralTensor* AuxiliaryTensor : AuxiliaryTensors)
+	{
+		AuxiliaryTensor->SetEnableGPU(true);
+	}
+	return true;
+}
+
+FNeuralTensor& FNeuralOperator::GetInputTensorNoConst()
+{
+	FPrivateNeuralOperator::SanityCheckInlining(true, InlinedTensor, Name);
+	checkf(PrivateInputTensors.Num() == 1, TEXT("FNeuralOperator::GetInputTensorNoConst(): Required 1 input, not %d. Call GetInputTensorsNoConst() instead."), PrivateInputTensors.Num());
+	return *PrivateInputTensors[0];
+}
+
+TArray<FNeuralTensor*>& FNeuralOperator::GetInputTensorsNoConst()
+{
+	FPrivateNeuralOperator::SanityCheckInlining(true, InlinedTensor, Name);
+	return PrivateInputTensors;
+}
+
+const FNeuralTensor& FNeuralOperator::GetInputTensorConst() const
+{
+	checkf(PrivateInputTensors.Num() == 1, TEXT("FNeuralOperator::GetInputTensorConst(): Required 1 input, not %d. Call GetInputTensorsConst() instead."), PrivateInputTensors.Num());
+	return *PrivateInputTensors[0];
+}
+
+FNeuralTensor& FNeuralOperator::GetOutputTensorNoConst()
+{
+	FPrivateNeuralOperator::SanityCheckInlining(false, InlinedTensor, Name);
+	checkf(PrivateOutputTensors.Num() == 1, TEXT("FNeuralOperator::GetOutputTensorNoConst(): Required 1 output, not %d. Call GetOutputTensorsNoConst() instead."), PrivateOutputTensors.Num());
+	return *PrivateOutputTensors[0];
+}
+
+TArray<FNeuralTensor*>& FNeuralOperator::GetOutputTensorsNoConst()
+{
+	FPrivateNeuralOperator::SanityCheckInlining(false, InlinedTensor, Name);
+	return PrivateOutputTensors;
+}
+
+const FNeuralTensor& FNeuralOperator::GetOutputTensorConst() const
+{
+	FPrivateNeuralOperator::SanityCheckInlining(false, InlinedTensor, Name);
+	checkf(PrivateOutputTensors.Num() == 1, TEXT("FNeuralOperator::GetOutputTensorConst(): Required 1 output, not %d. Call GetOutputTensorsConst() instead."), PrivateOutputTensors.Num());
+	return *PrivateOutputTensors[0];
+}
+
+const TArray<FNeuralTensor*>& FNeuralOperator::GetOutputTensorsConst() const
+{
+	FPrivateNeuralOperator::SanityCheckInlining(false, InlinedTensor, Name);
+	return PrivateOutputTensors;
+}
diff --git a/Engine/Plugins/Experimental/NeuralNetworkInference/Source/NeuralNetworkInference/Private/NeuralOperators/AbsOperator.cpp b/Engine/Plugins/Experimental/NeuralNetworkInference/Source/NeuralNetworkInference/Private/NeuralOperators/AbsOperator.cpp
new file mode 100644
--- /dev/null
+++ b/Engine/Plugins/Experimental/NeuralNetworkInference/Source/NeuralNetworkInference/Private/NeuralOperators/AbsOperator.cpp
@@ -0,0 +1,17 @@
+// Copyright Epic Games, Inc. All Rights Reserved.
+
+#include "NeuralOperators/AbsOperator.h"
+
+
+
+/* FAbsOperator structors
+ *****************************************************************************/
+
+FAbsOperator::FAbsOperator(const bool bIsInlinedTensor)
+	: IElementWiseOperator(TEXT("Abs"), 13, EElementWiseOperator::Abs, bIsInlinedTensor)
+{
+}
+
+FAbsOperator::~FAbsOperator()
+{
+}
diff --git a/Engine/Plugins/Experimental/NeuralNetworkInference/Source/NeuralNetworkInference/Private/NeuralOperators/AcosOperator.cpp b/Engine/Plugins/Experimental/NeuralNetworkInference/Source/NeuralNetworkInference/Private/NeuralOperators/AcosOperator.cpp
new file mode 100644
--- /dev/null
+++ b/Engine/Plugins/Experimental/NeuralNetworkInference/Source/NeuralNetworkInference/Private/NeuralOperators/AcosOperator.cpp
@@ -0,0 +1,17 @@
+// Copyright Epic Games, Inc. All Rights Reserved.
+
+#include "NeuralOperators/AcosOperator.h"
+
+
+
+/* FAcosOperator structors
+ *****************************************************************************/
+
+FAcosOperator::FAcosOperator(const bool bIsInlinedTensor)
+	: IElementWiseOperator(TEXT("Acos"), 13, EElementWiseOperator::Acos, bIsInlinedTensor)
+{
+}
+
+FAcosOperator::~FAcosOperator()
+{
+}
diff --git a/Engine/Plugins/Experimental/NeuralNetworkInference/Source/NeuralNetworkInference/Private/NeuralOperators/AddOperator.cpp b/Engine/Plugins/Experimental/NeuralNetworkInference/Source/NeuralNetworkInference/Private/NeuralOperators/AddOperator.cpp
new file mode 100644
--- /dev/null
+++ b/Engine/Plugins/Experimental/NeuralNetworkInference/Source/NeuralNetworkInference/Private/NeuralOperators/AddOperator.cpp
@@ -0,0 +1,17 @@
+// Copyright Epic Games, Inc. All Rights Reserved.
+
+#include "NeuralOperators/AddOperator.h"
+
+
+
+/* FAddOperator structors
+ *****************************************************************************/
+
+FAddOperator::FAddOperator(const TSet<uint32>& InPotentialInlinedTensors)
+	: IMultidirectionalBroadcastOperator(TEXT("Add"), 13, EMultidirectionalBroadcastOperator::Add, InPotentialInlinedTensors)
+{
+}
+
+FAddOperator::~FAddOperator()
+{
+}
diff --git a/Engine/Plugins/Experimental/NeuralNetworkInference/Source/NeuralNetworkInference/Private/NeuralOperators/AsinOperator.cpp b/Engine/Plugins/Experimental/NeuralNetworkInference/Source/NeuralNetworkInference/Private/NeuralOperators/AsinOperator.cpp
new file mode 100644
--- /dev/null
+++ b/Engine/Plugins/Experimental/NeuralNetworkInference/Source/NeuralNetworkInference/Private/NeuralOperators/AsinOperator.cpp
@@ -0,0 +1,17 @@
+// Copyright Epic Games, Inc. All Rights Reserved.
+
+#include "NeuralOperators/AsinOperator.h"
+
+
+
+/* FAsinOperator structors
+ *****************************************************************************/
+
+FAsinOperator::FAsinOperator(const bool bIsInlinedTensor)
+	: IElementWiseOperator(TEXT("Asin"), 7, EElementWiseOperator::Asin, bIsInlinedTensor)
+{
+}
+
+FAsinOperator::~FAsinOperator()
+{
+}
diff --git a/Engine/Plugins/Experimental/NeuralNetworkInference/Source/NeuralNetworkInference/Private/NeuralOperators/AtanOperator.cpp b/Engine/Plugins/Experimental/NeuralNetworkInference/Source/NeuralNetworkInference/Private/NeuralOperators/AtanOperator.cpp
new file mode 100644
--- /dev/null
+++ b/Engine/Plugins/Experimental/NeuralNetworkInference/Source/NeuralNetworkInference/Private/NeuralOperators/AtanOperator.cpp
@@ -0,0 +1,17 @@
+// Copyright Epic Games, Inc. All Rights Reserved.
+
+#include "NeuralOperators/AtanOperator.h"
+
+
+
+/* FAtanOperator structors
+ *****************************************************************************/
+
+FAtanOperator::FAtanOperator(const bool bIsInlinedTensor)
+	: IElementWiseOperator(TEXT("Atan"), 7, EElementWiseOperator::Atan, bIsInlinedTensor)
+{
+}
+
+FAtanOperator::~FAtanOperator()
+{
+}
diff --git a/Engine/Plugins/Experimental/NeuralNetworkInference/Source/NeuralNetworkInference/Private/NeuralOperators/BatchNormalizationOperator.cpp b/Engine/Plugins/Experimental/NeuralNetworkInference/Source/NeuralNetworkInference/Private/NeuralOperators/BatchNormalizationOperator.cpp
new file mode 100644
--- /dev/null
+++ b/Engine/Plugins/Experimental/NeuralNetworkInference/Source/NeuralNetworkInference/Private/NeuralOperators/BatchNormalizationOperator.cpp
@@ -0,0 +1,201 @@
+// Copyright Epic Games, Inc. All Rights Reserved.
+
+#include "NeuralOperators/BatchNormalizationOperator.h"
+#include "NeuralNetworkInferenceShaders/BatchNormalizationCS.h"
+#include "NeuralNetworkInferenceUtils.h"
+#include "NeuralNetworkInferenceUtilsGPU.h"
+
+
+
+/* FPrivateBatchNormalizationOperator static functions
+ *****************************************************************************/
+
+class FPrivateBatchNormalizationOperator
+{
+public:
+	static void ChannelBatchNormalization(FNeuralTensor& InOutOutputTensor, const TArray<FNeuralTensor*>& InInputTensors, const int64 InChannelIndex, const int64 InImageArea,
+		const int64 InAbsoluteIndexBias, const float InEpsilon);
+};
+
+void FPrivateBatchNormalizationOperator::ChannelBatchNormalization(FNeuralTensor& InOutOutputTensor, const TArray<FNeuralTensor*>& InInputTensors, const int64 InChannelIndex,
+	const int64 InImageArea, const int64 InAbsoluteIndexBias, const float InEpsilon)
+{
+	const float Scale = InInputTensors[1]->At<float>(InChannelIndex);
+	const float Bias = InInputTensors[2]->At<float>(InChannelIndex);
+	const float Mean = InInputTensors[3]->At<float>(InChannelIndex);
+	const float Variance = InInputTensors[4]->At<float>(InChannelIndex);
+	for (int64 AreaIndex = 0; AreaIndex < InImageArea; ++AreaIndex)
+	{
+		const int64 AbsoluteIndex = InAbsoluteIndexBias + AreaIndex;
+		const float X = InInputTensors[0]->At<float>(AbsoluteIndex);
+		InOutOutputTensor.At<float>(AbsoluteIndex) = Scale * (X - Mean) / FMath::Sqrt(Variance + InEpsilon) + Bias;
+	}
+}
+
+
+
+/* FBatchNormalizationOperator structors
+ *****************************************************************************/
+
+FBatchNormalizationOperator::FBatchNormalizationOperator(const bool bIsInlinedTensor, const FNodeProto& InNodeProto)
+	: FBatchNormalizationOperator(bIsInlinedTensor)
+{
+	if (const FAttributeProto* EpsilonAttribute = FModelProto::FindElementInArray(TEXT("Epsilon"), InNodeProto.Attribute, /*bMustValueBeFound*/false))
+	{
+		Epsilon = EpsilonAttribute->F;
+	}
+	if (const FAttributeProto* MomentumAttribute = FModelProto::FindElementInArray(TEXT("Momentum"), InNodeProto.Attribute, /*bMustValueBeFound*/false))
+	{
+		Momentum = MomentumAttribute->F;
+	}
+}
+
+FBatchNormalizationOperator::FBatchNormalizationOperator(const bool bIsInlinedTensor, const float InEpsilon, const float InMomentum)
+	: FNeuralOperator(TEXT("BatchNormalization"), 9, (bIsInlinedTensor ? 0 : -1))
+	, Epsilon(InEpsilon)
+	, Momentum(InMomentum)
+{
+}
+
+FBatchNormalizationOperator::~FBatchNormalizationOperator()
+{
+}
+
+
+
+/* FBatchNormalizationOperator public functions
+ *****************************************************************************/
+
+bool FBatchNormalizationOperator::ConfigureOutputAndInternalVariablesAndSanityChecks()
+{
+	bIsLoaded = false;
+	// Input sanity checks
+	const TArray<FNeuralTensor*>& InputTensors = GetInputTensorsConst();
+	if (!FNeuralNetworkInferenceUtils::SizeSanityChecks(InputTensors, 5))
+	{
+		return bIsLoaded;
+	}
+	for (int32 InputTensorIndex = 0; InputTensorIndex < InputTensors.Num(); ++InputTensorIndex)
+	{
+		FNeuralTensor& InputTensor = *InputTensors[InputTensorIndex];
+		if (InputTensor.GetNumberDimensions() != 1 && (InputTensorIndex != 0 || InputTensor.GetNumberDimensions() < 1))
+		{
+			UE_LOG(LogNeuralNetworkInference, Warning,
+				TEXT("FBatchNormalizationOperator::ConfigureOutputAndInternalVariablesAndSanityChecks(): InputTensor.GetNumberDimensions() is %d and not 1 (or also greater than 2 for index 0) failed for tensor index %d and name %s."),
+				InputTensor.GetNumberDimensions(), InputTensorIndex, *InputTensor.GetName());
+			return bIsLoaded;
+		}
+	}
+	// Output (if not inlined)
+	if (InlinedTensor < 0)
+	{
+		if (!FNeuralNetworkInferenceUtils::SizeSanityChecks(GetOutputTensorsConst(), 1))
+		{
+			return bIsLoaded;
+		}
+		FNeuralTensor& OutputTensor = GetOutputTensorNoConst();
+		if (OutputTensor.GetSizes() != InputTensors[0]->GetSizes())
+		{
+			OutputTensor.SetNumUninitialized(*InputTensors[0]);
+		}
+	}
+	// Return true
+	bIsLoaded = true;
+	return bIsLoaded;
+}
+
+void FBatchNormalizationOperator::ForwardCPU()
+{
+	// Sanity check
+	if (!bIsLoaded)
+	{
+		UE_LOG(LogNeuralNetworkInference, Warning, TEXT("FBatchNormalizationOperator::ForwardCPU(): bIsLoaded was false."));
+		return;
+	}
+
+	// Get input and output tensors
+	const TArray<FNeuralTensor*>& InputTensors = GetInputTensorsConst();
+	FNeuralTensor& OutputTensor = (InlinedTensor < 0 ? GetOutputTensorNoConst() : *GetInputTensorsNoConst()[0]);
+
+	// Apply batch normalization
+	// > 1 dimensions
+	if (OutputTensor.GetNumberDimensions() > 1)
+	{
+		const int64 BatchSize = OutputTensor.GetSize(0);
+		const int64 ChannelSize = OutputTensor.GetSize(1);
+		const int64 ChannelsVolume = OutputTensor.Num() / BatchSize;
+		const int64 ImageArea = ChannelsVolume / ChannelSize;
+		// Per batch
+		for (int64 NIndex = 0; NIndex < BatchSize; ++NIndex)
+		{
+			// Per channel
+			for (int64 ChannelIndex = 0; ChannelIndex < ChannelSize; ++ChannelIndex)
+			{
+				const int64 AbsoluteIndexBias = NIndex * ChannelsVolume + ChannelIndex * ImageArea;
+				FPrivateBatchNormalizationOperator::ChannelBatchNormalization(OutputTensor, InputTensors, ChannelIndex, ImageArea, AbsoluteIndexBias, Epsilon);
+			}
+		}
+	}
+	// 1 dimension
+	else if (OutputTensor.GetNumberDimensions() == 1)
+	{
+		FPrivateBatchNormalizationOperator::ChannelBatchNormalization(OutputTensor, InputTensors, /*Channel*/0, /*ImageArea*/OutputTensor.Num(), /*AbsoluteIndexBias*/0, Epsilon);
+	}
+	// Unknown
+	else
+	{
+		UE_LOG(LogNeuralNetworkInference, Warning, TEXT("FBatchNormalizationOperator::ForwardCPU(): OutputTensor.GetNumberDimensions() = %d but should be >= 1."), OutputTensor.GetNumberDimensions());
+	}
+}
+
+void FBatchNormalizationOperator::ForwardGPU_RenderThread(FRDGBuilder* InOutGraphBuilder)
+{
+	// Sanity checks
+	if (!FNeuralNetworkInferenceUtilsGPU::GPUSanityChecks(InOutGraphBuilder, bIsLoaded))
+	{
+		UE_LOG(LogNeuralNetworkInference, Warning, TEXT("FBatchNormalizationOperator::ForwardGPU_RenderThread(): Sanity checks failed."));
+		return;
+	}
+	// Set parameters
+	FNeuralTensor& OutputTensor = (InlinedTensor < 0 ? GetOutputTensorNoConst() : *GetInputTensorsNoConst()[0]);
+	const EBatchNormalizationMode BatchNormalizationMode = (OutputTensor.GetNumberDimensions() > 1 ? EBatchNormalizationMode::NDimensions : EBatchNormalizationMode::OneDimension);
+	FBatchNormalizationCS::FParameters* Parameters = InOutGraphBuilder->AllocParameters<FBatchNormalizationCS::FParameters>();
+	// > 1 dimensions
+	if (BatchNormalizationMode == EBatchNormalizationMode::NDimensions)
+	{
+		Parameters->BatchSize = OutputTensor.GetSize(0);
+		Parameters->ChannelSize = OutputTensor.GetSize(1);
+		Parameters->ChannelsVolume = OutputTensor.Num() / Parameters->BatchSize;
+		Parameters->ImageArea = Parameters->ChannelsVolume / Parameters->ChannelSize;
+	}
+	else
+	{
+		Parameters->ImageArea = OutputTensor.Num();
+	}
+	Parameters->Epsilon = Epsilon;
+	// SRV/UAV variables
+	const TArray<FNeuralTensor*>& InputTensors = GetInputTensorsConst();
+	const bool bIsInlined = (InlinedTensor > -1);
+	if (!bIsInlined)
+	{
+		Parameters->XSRV = InputTensors[0]->GetBufferSRVRef();
+	}
+	Parameters->OutputUAV = OutputTensor.GetBufferUAVRef();
+	Parameters->Scale = InputTensors[1]->GetBufferSRVRef();
+	Parameters->Bias = InputTensors[2]->GetBufferSRVRef();
+	Parameters->Mean = InputTensors[3]->GetBufferSRVRef();
+	Parameters->Variance = InputTensors[4]->GetBufferSRVRef();
+	// Set shader
+	FBatchNormalizationCS::FPermutationDomain PermutationVector;
+	PermutationVector.Set<FBatchNormalizationCS::FBatchNormalizationMode>(BatchNormalizationMode);
+	PermutationVector.Set<FBatchNormalizationCS::FIsInlined>(bIsInlined);
+	TShaderMapRef<FBatchNormalizationCS> ComputeShader(GetGlobalShaderMap(GMaxRHIFeatureLevel), PermutationVector);
+	// Run shader
+	const uint32 ThreadGroupCountValueX = FMath::DivideAndRoundUp(Parameters->ImageArea, FBatchNormalizationCS::THREADGROUP_SIZE_X);
+	FComputeShaderUtils::AddPass(
+		*InOutGraphBuilder,
+		RDG_EVENT_NAME("FBatchNormalizationCS() - Operator: %s, EBatchNormalizationMode: %d", *Name, BatchNormalizationMode),
+		ComputeShader,
+		Parameters,
+		FIntVector(ThreadGroupCountValueX, 1, 1));
+}
diff --git a/Engine/Plugins/Experimental/NeuralNetworkInference/Source/NeuralNetworkInference/Private/NeuralOperators/CeilOperator.cpp b/Engine/Plugins/Experimental/NeuralNetworkInference/Source/NeuralNetworkInference/Private/NeuralOperators/CeilOperator.cpp
new file mode 100644
--- /dev/null
+++ b/Engine/Plugins/Experimental/NeuralNetworkInference/Source/NeuralNetworkInference/Private/NeuralOperators/CeilOperator.cpp
@@ -0,0 +1,17 @@
+// Copyright Epic Games, Inc. All Rights Reserved.
+
+#include "NeuralOperators/CeilOperator.h"
+
+
+
+/* FCeilOperator structors
+ *****************************************************************************/
+
+FCeilOperator::FCeilOperator(const bool bIsInlinedTensor)
+	: IElementWiseOperator(TEXT("Ceil"), 13, EElementWiseOperator::Ceil, bIsInlinedTensor)
+{
+}
+
+FCeilOperator::~FCeilOperator()
+{
+}
diff --git a/Engine/Plugins/Experimental/NeuralNetworkInference/Source/NeuralNetworkInference/Private/NeuralOperators/ConstantOperator.cpp b/Engine/Plugins/Experimental/NeuralNetworkInference/Source/NeuralNetworkInference/Private/NeuralOperators/ConstantOperator.cpp
new file mode 100644
--- /dev/null
+++ b/Engine/Plugins/Experimental/NeuralNetworkInference/Source/NeuralNetworkInference/Private/NeuralOperators/ConstantOperator.cpp
@@ -0,0 +1,132 @@
+// Copyright Epic Games, Inc. All Rights Reserved.
+
+#include "NeuralOperators/ConstantOperator.h"
+#include "NeuralNetworkInferenceUtils.h"
+
+
+
+/* FConstantOperator structors
+ *****************************************************************************/
+
+FConstantOperator::FConstantOperator(const FNodeProto& InNodeProto)
+	: FConstantOperator(FNeuralTensor())
+{
+	// Fill tensor from one of the attributes
+	if (const FAttributeProto* SparseValueAttribute = FModelProto::FindElementInArray(TEXT("sparse_value"), InNodeProto.Attribute, /*bMustValueBeFound*/false))
+	{
+		UE_LOG(LogNeuralNetworkInference, Warning, TEXT("FConstantOperator::FConstantOperator(): \"sparse_value\" not implemented yet."));
+	}
+	else if (const FAttributeProto* ValueAttribute = FModelProto::FindElementInArray(TEXT("value"), InNodeProto.Attribute, /*bMustValueBeFound*/false))
+	{
+		if (ValueAttribute->Type == EAttributeProtoAttributeType::TENSOR)
+		{
+			// Get new tensor from TensorProto
+			if (!Tensor.SetFromTensorProto(&ValueAttribute->T, TEXT("FConstantOperator-") + ValueAttribute->Name, ENeuralTensorTypeGPU::Weight))
+			{
+				UE_LOG(LogNeuralNetworkInference, Warning, TEXT("FConstantOperator::FConstantOperator(): SetFromTensorProto() failed."));
+			}
+		}
+		// Sanity check
+		else
+		{
+			UE_LOG(LogNeuralNetworkInference, Warning, TEXT("FConstantOperator::FConstantOperator(): Not implemented type used, ValueAttribute->Type = %d."), (int32)ValueAttribute->Type);
+		}
+	}
+	else if (const FAttributeProto* ValueFloatAttribute = FModelProto::FindElementInArray(TEXT("value_float"), InNodeProto.Attribute, /*bMustValueBeFound*/false))
+	{
+		UE_LOG(LogNeuralNetworkInference, Warning, TEXT("FConstantOperator::FConstantOperator(): \"value\" not implemented yet."));
+	}
+	else if (const FAttributeProto* ValueFloatsAttribute = FModelProto::FindElementInArray(TEXT("value_floats"), InNodeProto.Attribute, /*bMustValueBeFound*/false))
+	{
+		UE_LOG(LogNeuralNetworkInference, Warning, TEXT("FConstantOperator::FConstantOperator(): \"value_floats\" not implemented yet."));
+	}
+	else if (const FAttributeProto* ValueIntAttribute = FModelProto::FindElementInArray(TEXT("value_int"), InNodeProto.Attribute, /*bMustValueBeFound*/false))
+	{
+		UE_LOG(LogNeuralNetworkInference, Warning, TEXT("FConstantOperator::FConstantOperator(): \"value_int\" not implemented yet."));
+	}
+	else if (const FAttributeProto* ValueIntsAttribute = FModelProto::FindElementInArray(TEXT("value_ints"), InNodeProto.Attribute,/*bMustValueBeFound*/ false))
+	{
+		UE_LOG(LogNeuralNetworkInference, Warning, TEXT("FConstantOperator::FConstantOperator(): \"value_ints\" not implemented yet."));
+	}
+	else if (const FAttributeProto* ValueStringAttribute = FModelProto::FindElementInArray(TEXT("value_string"), InNodeProto.Attribute, /*bMustValueBeFound*/false))
+	{
+		UE_LOG(LogNeuralNetworkInference, Warning, TEXT("FConstantOperator::FConstantOperator(): \"value_string\" not implemented yet."));
+	}
+	else if (const FAttributeProto* ValueStringsAttribute = FModelProto::FindElementInArray(TEXT("value_strings"), InNodeProto.Attribute, /*bMustValueBeFound*/false))
+	{
+		UE_LOG(LogNeuralNetworkInference, Warning, TEXT("FConstantOperator::FConstantOperator(): \"value_strings\" not implemented yet."));
+	}
+	else
+	{
+		UE_LOG(LogNeuralNetworkInference, Warning, TEXT("FConstantOperator::FConstantOperator(): No valid attributes found in NodeProto."));
+	}
+}
+
+FConstantOperator::FConstantOperator(const FNeuralTensor& InTensor)
+	: FNeuralOperator(TEXT("Constant"), 13, -1)
+	, Tensor(InTensor)
+{
+}
+
+FConstantOperator::~FConstantOperator()
+{
+}
+
+
+
+/* FConstantOperator public functions
+ *****************************************************************************/
+
+bool FConstantOperator::ConfigureOutputAndInternalVariablesAndSanityChecks()
+{
+	bIsLoaded = false;
+	// InlinedTensor must be false
+	if (InlinedTensor > -1)
+	{
+		UE_LOG(LogNeuralNetworkInference, Warning, TEXT("FConstantOperator::ConfigureOutputAndInternalVariablesAndSanityChecks(): InlinedTensor must be -1, but found %d."), InlinedTensor);
+		return bIsLoaded;
+	}
+	// Input sanity checks
+	const int32 NumberInputTensors = GetInputTensorsConst().Num();
+	if (NumberInputTensors > 0)
+	{
+		UE_LOG(LogNeuralNetworkInference, Warning, TEXT("FConstantOperator::ConfigureOutputAndInternalVariablesAndSanityChecks(): There cannot be any input tensors, but found %d."), NumberInputTensors);
+		return bIsLoaded;
+	}
+	// Output sanity checks
+	const int32 NumberOutputTensors = GetOutputTensorsConst().Num();
+	if (NumberOutputTensors != 1)
+	{
+		UE_LOG(LogNeuralNetworkInference, Warning, TEXT("FConstantOperator::ConfigureOutputAndInternalVariablesAndSanityChecks(): There must be exactly 1 output tensor, but found %d."), NumberOutputTensors);
+		return bIsLoaded;
+	}
+	// Variable sanity check
+	if (Tensor.IsEmpty())
+	{
+		UE_LOG(LogNeuralNetworkInference, Warning, TEXT("FConstantOperator::ConfigureOutputAndInternalVariablesAndSanityChecks(): Tensor cannot be empty."));
+		return bIsLoaded;
+	}
+	// Return true
+	bIsLoaded = true;
+	return bIsLoaded;
+}
+
+void FConstantOperator::ForwardCPU()
+{
+	// Sanity check
+	if (!bIsLoaded)
+	{
+		UE_LOG(LogNeuralNetworkInference, Warning, TEXT("FConstantOperator::ForwardCPU(): bIsLoaded was false."));
+		return;
+	}
+
+	// Fill output tensor
+	FNeuralTensor& OutputTensor = GetOutputTensorNoConst();
+	OutputTensor.SetNumUninitialized(Tensor);
+	OutputTensor.SetFromUnderlyingUInt8ArrayCopy(Tensor.GetUnderlyingUInt8ArrayRef());
+}
+
+void FConstantOperator::ForwardGPU_RenderThread(FRDGBuilder* InOutGraphBuilder)
+{
+	UE_LOG(LogNeuralNetworkInference, Warning, TEXT("FConstantOperator::ForwardGPU_RenderThread(): Constant does not implement the GPU forward pass, only CPU one."));
+}
diff --git a/Engine/Plugins/Experimental/NeuralNetworkInference/Source/NeuralNetworkInference/Private/NeuralOperators/ConvBaseOperator.cpp b/Engine/Plugins/Experimental/NeuralNetworkInference/Source/NeuralNetworkInference/Private/NeuralOperators/ConvBaseOperator.cpp
new file mode 100644
--- /dev/null
+++ b/Engine/Plugins/Experimental/NeuralNetworkInference/Source/NeuralNetworkInference/Private/NeuralOperators/ConvBaseOperator.cpp
@@ -0,0 +1,1019 @@
+// Copyright Epic Games, Inc. All Rights Reserved.
+
+#include "NeuralOperators/ConvBaseOperator.h"
+#include "NeuralNetworkInferenceShaders/ConvBaseCS.h"
+#include "NeuralNetworkInferenceUtils.h"
+#include "NeuralNetworkInferenceUtilsGPU.h"
+
+
+
+/* FPrivateConvBaseOperator auxiliary functions
+ *****************************************************************************/
+
+class FPrivateConvBaseOperator
+{
+public:
+	static void NDTensorIndexesPlus1(TArray<int64>& InOutImageAreaIndexes, const TArray<int64>& InSizes);
+
+	/**
+	 * It will return an empty FNeuralInt64ArrayUInt32Buffer if any error occurred.
+	 */
+	static FNeuralInt64ArrayUInt32Buffer TransposePadsFromConvPads(const FNeuralInt64ArrayUInt32Buffer& InPads, const FConvBaseOperator::EAutoPad InAutoPad, const TArray<int64>& InEffectiveKernelSizes);
+};
+
+void FPrivateConvBaseOperator::NDTensorIndexesPlus1(TArray<int64>& InOutImageAreaIndexes, const TArray<int64>& InSizes)
+{
+	int64 ImageAreaIndexesIndex = InSizes.Num() - 3; // = NumberConvolutionalDimensions - 1 = InSizes.Num() - 2 - 1
+	while (ImageAreaIndexesIndex > -1)
+	{
+		++InOutImageAreaIndexes[ImageAreaIndexesIndex];
+		if (InOutImageAreaIndexes[ImageAreaIndexesIndex] == InSizes[ImageAreaIndexesIndex + 2])
+		{
+			InOutImageAreaIndexes[ImageAreaIndexesIndex] = 0;
+			--ImageAreaIndexesIndex;
+		}
+		else
+		{
+			break;
+		}
+	}
+}
+
+FNeuralInt64ArrayUInt32Buffer FPrivateConvBaseOperator::TransposePadsFromConvPads(const FNeuralInt64ArrayUInt32Buffer& InPads, const FConvBaseOperator::EAutoPad InAutoPad, const TArray<int64>& InEffectiveKernelSizes)
+{
+	FNeuralInt64ArrayUInt32Buffer Pads;
+	TArray<int64> PadsArrayTemporary;
+	PadsArrayTemporary.SetNumUninitialized(InPads.Num());
+	for (int32 Index = 0; Index < PadsArrayTemporary.Num(); ++Index)
+	{
+		PadsArrayTemporary[Index] = InEffectiveKernelSizes[Index / 2] - InPads[Index] - 1;
+		if (PadsArrayTemporary[Index] < 0)
+		{
+			UE_LOG(LogNeuralNetworkInference, Warning, TEXT("FConvTransposeOperator: PadsArrayTemporary[%d] = %d, ConvTranspose only implemented for PadsArrayTemporary[%d] >= 0 for now."), Index, PadsArrayTemporary[Index], Index);
+			return Pads;
+		}
+	}
+	Pads.Move(PadsArrayTemporary);
+	return Pads;
+}
+
+
+
+/* FConvBaseOperator structors
+ *****************************************************************************/
+
+FConvBaseOperator::FConvBaseOperator(const FString& InName, const int32 InVersion, const int32 InInlinedTensor,
+	const EAutoPad InAutoPad, const TArray<int64>& InDilations, const int64 InGroup, const TArray<int64>& InKernelShape, const TArray<int64>& InPads, const TArray<int64>& InStrides,
+	const bool bInIsTransposed, const TArray<int64>& InOutputPadding, const TArray<int64>& InOutputShape)
+	: FNeuralOperator(InName, InVersion, InInlinedTensor)
+	, bIsTransposed(bInIsTransposed)
+	, AutoPad(InAutoPad)
+	, Dilations(InDilations)
+	, Group(InGroup)
+	, KernelShape(InKernelShape)
+	, Pads(InPads)
+	, Strides(InStrides)
+	, OutputPadding(InOutputPadding)
+	, OutputShape(InOutputShape)
+{
+}
+
+FConvBaseOperator::~FConvBaseOperator()
+{
+}
+
+
+
+/* FConvBaseOperator public functions
+ *****************************************************************************/
+
+bool FConvBaseOperator::ConfigureOutputAndInternalVariablesAndSanityChecks()
+{
+	bIsLoaded = false;
+	// Input sanity checks
+	const TArray<FNeuralTensor*>& InputTensors = GetInputTensorsConst();
+	const TArray<int64>& XSizes = InputTensors[0]->GetSizes();
+	const TArray<int64>& WSizes = InputTensors[1]->GetSizes();
+	if (!FNeuralNetworkInferenceUtils::SizeSanityChecks(InputTensors, 2, 3, 3, -1, 0, 1) || !FNeuralNetworkInferenceUtils::SizeSanityChecks(InputTensors, 2, 3, 1, 1, 2, 2))
+	{
+		return false;
+	}
+	// Initialize and set first 2 dimensions of PartiallySetOutputSizes + sanity checks
+	TArray<int64> PartiallySetOutputSizes;
+	{
+		if (!FNeuralNetworkInferenceUtils::SizeSanityChecks(GetOutputTensorsConst(), 1, 1, -1, -1))
+		{
+			return bIsLoaded;
+		}
+		PartiallySetOutputSizes.SetNumUninitialized(XSizes.Num());
+		PartiallySetOutputSizes[0] = XSizes[0];
+		PartiallySetOutputSizes[1] = (!bIsTransposed ? WSizes[0] : Group * WSizes[1]);
+	}
+	// Sanity checks
+	if (!SanityChecksForTensors(PartiallySetOutputSizes))
+	{
+		return bIsLoaded;
+	}
+	const int32 NumberDimensions = InputTensors[1]->GetNumberDimensions();
+	const int32 NumberConvolutionalDimensions = NumberDimensions - 2;
+	// Attribute sanity checks
+	// If KernelShape given, make sure it matches W
+	if (KernelShape.Num() > 0)
+	{
+		if (KernelShape.Num() != NumberConvolutionalDimensions)
+		{
+			UE_LOG(LogNeuralNetworkInference, Warning,
+				TEXT("FConvBaseOperator::ConfigureOutputAndInternalVariablesAndSanityChecks(): KernelShape.Num() == NumberConvolutionalDimensions failed (%d != %d)."),
+				KernelShape.Num(), NumberConvolutionalDimensions);
+			return bIsLoaded;
+		}
+		for (int64 DimensionIndex = 2; DimensionIndex < NumberDimensions; ++DimensionIndex)
+		{
+			if (WSizes[DimensionIndex] != KernelShape[DimensionIndex - 2])
+			{
+				UE_LOG(LogNeuralNetworkInference, Warning,
+					TEXT("FConvBaseOperator::ConfigureOutputAndInternalVariablesAndSanityChecks(): KernelShape[%d] does not match the actual size of W[%d] (%d != %d)."),
+					DimensionIndex, DimensionIndex - 2, WSizes[DimensionIndex], KernelShape[DimensionIndex - 2]);
+				return bIsLoaded;
+			}
+		}
+	}
+	// OutputPadding sanity check
+	if (OutputPadding.Num() != 0 && OutputPadding.Num() != NumberConvolutionalDimensions)
+	{
+		UE_LOG(LogNeuralNetworkInference, Warning, TEXT("FConvBaseOperator::ConfigureOutputAndInternalVariablesAndSanityChecks(): OutputPadding.Num() should be 0 or %d, but not %d."), NumberConvolutionalDimensions, OutputPadding.Num());
+		return bIsLoaded;
+	}
+	// OutputShape sanity check
+	if (OutputShape.Num() != 0 && OutputShape.Num() != NumberConvolutionalDimensions)
+	{
+		UE_LOG(LogNeuralNetworkInference, Warning, TEXT("FConvBaseOperator::ConfigureOutputAndInternalVariablesAndSanityChecks(): OutputShape.Num() should be 0 or %d, but not %d."), NumberConvolutionalDimensions, OutputShape.Num());
+		return bIsLoaded;
+	}
+	// OutputShape/Pads sanity check --> Only 1 of the 2
+	if (OutputShape.Num() > 0 && Pads.Num() > 0)
+	{
+		for (int64 DimensionIndex = 0; DimensionIndex < FMath::Min(NumberConvolutionalDimensions, Pads.Num()); ++DimensionIndex) // Min to avoid crash if wrong Pad (wrong Pad checked later)
+		{
+			if (OutputShape[DimensionIndex] > 0 && Pads[DimensionIndex] > 0)
+			{
+				UE_LOG(LogNeuralNetworkInference, Warning, TEXT("FConvBaseOperator::ConfigureOutputAndInternalVariablesAndSanityChecks(): You can only set OutputShape or Pads, but not both. Pads will be overridden."));
+				Pads = FNeuralInt64ArrayUInt32Buffer();
+				break;
+			}
+		}
+	}
+	if (!FillOutputShape(XSizes))
+	{
+		UE_LOG(LogNeuralNetworkInference, Warning, TEXT("FConvBaseOperator::ConfigureOutputAndInternalVariablesAndSanityChecks(): FillOutputShape() failed."));
+		return bIsLoaded;
+	}
+	// Set Dilations
+	if (!FillDilationsIfEmpty(NumberConvolutionalDimensions))
+	{
+		UE_LOG(LogNeuralNetworkInference, Warning, TEXT("FConvBaseOperator::ConfigureOutputAndInternalVariablesAndSanityChecks(): FillDilationsIfEmpty() failed."));
+		return bIsLoaded;
+	}
+	// Set Strides
+	if (!SetAndConfigureStrides(NumberConvolutionalDimensions))
+	{
+		UE_LOG(LogNeuralNetworkInference, Warning, TEXT("FConvBaseOperator::ConfigureOutputAndInternalVariablesAndSanityChecks(): SetAndConfigureStrides() failed."));
+		return bIsLoaded;
+	}
+	// Strides sanity checks
+	if (Strides.Num() > 0 && Strides.Num() != NumberConvolutionalDimensions) // Right dimensions
+	{
+		UE_LOG(LogNeuralNetworkInference, Warning, TEXT("FConvBaseOperator::ConfigureOutputAndInternalVariablesAndSanityChecks(): Strides.Num() == 0 || Strides.Num() == NumberConvolutionalDimensions failed (%d != 0 || %d)."),
+			Strides.Num(), NumberConvolutionalDimensions);
+		return bIsLoaded;
+	}
+	// Set AuxiliaryTensors if needed
+	if (!SetAndConfigureAuxiliaryTensor())
+	{
+		UE_LOG(LogNeuralNetworkInference, Warning, TEXT("FConvBaseOperator::ConfigureOutputAndInternalVariablesAndSanityChecks(): SetAndConfigureAuxiliaryTensor() failed."));
+		return bIsLoaded;
+	}
+	for (const int64 Stride : Strides) // Strides >= 1 on each dimension
+	{
+		if (Stride < 1)
+		{
+			UE_LOG(LogNeuralNetworkInference, Warning, TEXT("FConvBaseOperator::ConfigureOutputAndInternalVariablesAndSanityChecks(): Stride (%d) < 1 found! Stride must be >= 1."), Stride);
+			return bIsLoaded;
+		}
+	}
+	// Create EffectiveKernelSizes
+	const TArray<int64> EffectiveKernelSizes = CreateEffectiveKernelSizes(NumberConvolutionalDimensions, WSizes);
+	// Either AutoPad == NotSet with explicit Pads or AutoPad != NotSet without explicit Pads. Both together is not possible
+	if (AutoPad != EAutoPad::NotSet && Pads.Num() != 0)
+	{
+		// Make sure Pads != {0,0,...,0} (or it would be equivalent to Pads.Num() = 0, thus valid)
+		for (const int64 Pad : Pads)
+		{
+			if (Pad != 0)
+			{
+				UE_LOG(LogNeuralNetworkInference, Warning, TEXT("FConvOperator::ConfigureOutputAndInternalVariablesAndSanityChecks(): If AutoPad (%d) is different than NotSet (%d), explicit padding is not allowed."),
+					(int32)AutoPad, (int32)EAutoPad::NotSet);
+				return bIsLoaded;
+			}
+		}
+	}
+	// Sanity check
+	if (AuxiliaryTensors.Num() > 1)
+	{
+		UE_LOG(LogNeuralNetworkInference, Warning, TEXT("FConvTransposeOperator::SetAndConfigureAuxiliaryTensor(): AuxiliaryTensors.Num() should be 0 or 1, not %d."), AuxiliaryTensors.Num());
+		return false;
+	}
+	// Set Pads (based on InPads and InAutoPad)
+	const FNeuralTensor& XOrXWithZeros = *(AuxiliaryTensors.Num() == 0 ? InputTensors[0] : AuxiliaryTensors[0]);
+	const TArray<int64>& XOrXWithZerosSizes = XOrXWithZeros.GetSizes();
+	if (!EstimateAndFillPads(EffectiveKernelSizes, NumberDimensions, XSizes))
+	{
+		UE_LOG(LogNeuralNetworkInference, Warning, TEXT("FConvBaseOperator::ConfigureOutputAndInternalVariablesAndSanityChecks(): EstimateAndFillPads() failed."));
+		return bIsLoaded;
+	}
+	// Pads size sanity checks
+	if (Pads.Num() > 0 && Pads.Num() != 2 * NumberConvolutionalDimensions) // Right dimensions
+	{
+		UE_LOG(LogNeuralNetworkInference, Warning,
+			TEXT("FConvOperator::ConfigureOutputAndInternalVariablesAndSanityChecks(): Pads.Num() == 0 || Pads.Num() == 2 * NumberConvolutionalDimensions failed (%d != 0 || %d)."),
+			Pads.Num(), 2 * NumberConvolutionalDimensions);
+		return bIsLoaded;
+	}
+	for (const int64 Pad : Pads) // Pads >= 0 on each dimension
+	{
+		if (Pad < 0)
+		{
+			UE_LOG(LogNeuralNetworkInference, Warning, TEXT("FConvOperator::ConfigureOutputAndInternalVariablesAndSanityChecks(): Pads >= 0 failed (found %d)."), Pad);
+			return bIsLoaded;
+		}
+	}
+	// Final sanity checks
+	if (!FillOutputSizesAndPaddedMarginsAndSanityChecks(PartiallySetOutputSizes, EffectiveKernelSizes))
+	{
+		UE_LOG(LogNeuralNetworkInference, Warning, TEXT("FConvBaseOperator::ConfigureOutputAndInternalVariablesAndSanityChecks(): FillOutputSizesAndPaddedMarginsAndSanityChecks() failed."));
+		return bIsLoaded;
+	}
+	// Return true
+	bIsLoaded = true;
+	return bIsLoaded;
+}
+
+void FConvBaseOperator::ToGPU_RenderThread()
+{
+	// Sanity check
+	if (OutputSizesForGPU.IsReadBufferValid() || XOrXWithZerosSizesForGPU.IsReadBufferValid() || WSizesForGPU.IsReadBufferValid()
+		|| Dilations.IsReadBufferValid() || Strides.IsReadBufferValid() || Pads.IsReadBufferValid())
+	{
+		UE_LOG(LogNeuralNetworkInference, Warning, TEXT("FConvBaseOperator::ToGPU_RenderThread():"
+			" OutputSizesForGPU, XOrXWithZerosSizesForGPU, WSizesForGPU, OutputSizesForGPU, XOrXWithZerosSizesForGPU, WSizesForGPU, Dilations, Strides, and/or Pads FReadBuffers were not nullptr."));
+	}
+	// Get input and output tensors
+	const TArray<FNeuralTensor*>& InputTensors = GetInputTensorsConst();
+	FNeuralTensor& OutputTensor = GetOutputTensorNoConst();
+	// Memory to GPU
+	OutputSizesForGPU.Copy(OutputTensor.GetSizes());
+	const FNeuralTensor& XOrXWithZeros = *(AuxiliaryTensors.Num() == 0 ? InputTensors[0] : AuxiliaryTensors[0]);
+	XOrXWithZerosSizesForGPU.Copy(XOrXWithZeros.GetSizes());
+	WSizesForGPU.Copy(InputTensors[1]->GetSizes());
+	OutputSizesForGPU.CreateAndLoadSRVBuffer(TEXT("Conv_ReadBufferOutputSizes"));
+	XOrXWithZerosSizesForGPU.CreateAndLoadSRVBuffer(TEXT("Conv_ReadBufferXSizes"));
+	WSizesForGPU.CreateAndLoadSRVBuffer(TEXT("Conv_ReadBufferWSizes"));
+	Dilations.CreateAndLoadSRVBuffer(TEXT("Conv_ReadBufferDilations"));
+	Strides.CreateAndLoadSRVBuffer(TEXT("Conv_ReadBufferStrides"));
+	Pads.CreateAndLoadSRVBuffer(TEXT("Conv_ReadBufferPads"));
+}
+
+void FConvBaseOperator::ForwardCPU()
+{
+	// Sanity check
+	if (!bIsLoaded)
+	{
+		UE_LOG(LogNeuralNetworkInference, Warning, TEXT("FConvBaseOperator::ForwardCPU(): bIsLoaded was false."));
+		return;
+	}
+
+	// Get input and output tensors
+	const TArray<FNeuralTensor*>& InputTensors = GetInputTensorsConst();
+	FNeuralTensor& OutputTensor = GetOutputTensorNoConst();
+
+	// Set output to bias (or 0)
+	if (InputTensors.Num() > 2)
+	{
+		const TArray<int64>& OutputSizes = OutputTensor.GetSizes();
+		const int64 OutputBatchVolume = OutputTensor.Num() / OutputSizes[0];
+		const int64 OutputImageArea = OutputBatchVolume / OutputSizes[1];
+		const FNeuralTensor& B = *InputTensors[2];
+		// Per batch N
+		for (int64 BatchIndex = 0; BatchIndex < OutputSizes[0]; ++BatchIndex)
+		{
+			// Per output channel M
+			for (int64 OutputChannelIndex = 0; OutputChannelIndex < OutputSizes[1]; ++OutputChannelIndex)
+			{
+				// Per image area (D1 x D2 x ... x Dn)
+				for (int64 OutputImageIndex = 0; OutputImageIndex < OutputImageArea; ++OutputImageIndex)
+				{
+					OutputTensor.At<float>(BatchIndex * OutputBatchVolume + OutputChannelIndex * OutputImageArea + OutputImageIndex) = B.At<float>(OutputChannelIndex);
+				}
+			}
+		}
+	}
+	else
+	{
+		OutputTensor.SetTo<float>(0);
+	}
+
+	// Apply convolution
+	const FNeuralTensor& W = *InputTensors[1];
+	const FNeuralTensor& XOrXWithZeros = *(AuxiliaryTensors.Num() == 0 ? InputTensors[0] : AuxiliaryTensors[0]);
+	MultiChannelConvolutionCPU(OutputTensor, XOrXWithZeros, W);
+}
+
+void FConvBaseOperator::ForwardGPU_RenderThread(FRDGBuilder* InOutGraphBuilder)
+{
+	// Sanity checks
+	if (!FNeuralNetworkInferenceUtilsGPU::GPUSanityChecks(InOutGraphBuilder, bIsLoaded))
+	{
+		UE_LOG(LogNeuralNetworkInference, Warning, TEXT("FConvBaseOperator::ForwardGPU_RenderThread(): Sanity checks failed."));
+		return;
+	}
+	FNeuralTensor& OutputTensor = GetOutputTensorNoConst();
+	const int32 NumberConvolutionalDimensions = OutputTensor.GetNumberDimensions() - 2;
+	if (NumberConvolutionalDimensions > (int32)FConvBaseCS::MAX_NUMBER_DIMENSIONS)
+	{
+		UE_LOG(LogNeuralNetworkInference, Warning, TEXT("FConvBaseOperator::ForwardGPU_RenderThread(): NumberConvolutionalDimensions <= FConvBaseCS::MAX_NUMBER_DIMENSIONS failed (%d vs. %d). I.e.,"
+			" the maximum number of dimensions for FConvBaseOperator is FConvBaseCS::MAX_NUMBER_DIMENSIONS. Increase FConvBaseCS::MAX_NUMBER_DIMENSIONS if you need a higher dimensional convolution,"
+			" which might result in a performance hit for all convolutions of more than 3 dimensions."),
+			NumberConvolutionalDimensions, FConvBaseCS::MAX_NUMBER_DIMENSIONS);
+		return;
+	}
+	// Set parameters
+	FConvBaseCS::FParameters* Parameters = InOutGraphBuilder->AllocParameters<FConvBaseCS::FParameters>();
+	// Input SRV variables - Fill ReadBuffer's the first time
+	const TArray<FNeuralTensor*>& InputTensors = GetInputTensorsConst();
+	// Input SRV variables
+	Parameters->OutputSizes = OutputSizesForGPU.GetUInt32SRV();
+	Parameters->XOrXWithZerosSizes = XOrXWithZerosSizesForGPU.GetUInt32SRV();
+	Parameters->WSizes = WSizesForGPU.GetUInt32SRV();
+	Parameters->Dilations = Dilations.GetUInt32SRV();
+	Parameters->Strides = Strides.GetUInt32SRV();
+	Parameters->Pads = Pads.GetUInt32SRV();
+	// Input variables
+	Parameters->bIsTransposed = bIsTransposed;
+	Parameters->Group = Group;
+	Parameters->MIntoGroup = (!bIsTransposed ? WSizesForGPU[0] / Group : WSizesForGPU[1]);
+	Parameters->CIntoGroup = WSizesForGPU[0] / Group;
+	Parameters->OutputVolume = OutputTensor.Num();
+	Parameters->OutputBatchVolume = Parameters->OutputVolume / OutputSizesForGPU[0];
+	Parameters->OutputImageArea = Parameters->OutputBatchVolume / OutputSizesForGPU[1];
+	Parameters->WBatchVolume = InputTensors[1]->Num() / WSizesForGPU[0];
+	Parameters->WImageArea = Parameters->WBatchVolume / WSizesForGPU[1];
+	const FNeuralTensor& XOrXWithZeros = *(AuxiliaryTensors.Num() == 0 ? InputTensors[0] : AuxiliaryTensors[0]);
+	Parameters->XOrXWithZerosBatchVolume = XOrXWithZeros.Num() / XOrXWithZerosSizesForGPU[0];
+	Parameters->XOrXWithZerosImageArea = Parameters->XOrXWithZerosBatchVolume / XOrXWithZerosSizesForGPU[1];
+	// nD convolution
+	const EConvMode ConvMode = (NumberConvolutionalDimensions == 1 ? EConvMode::D1
+		: NumberConvolutionalDimensions == 2 ? EConvMode::D2
+		: NumberConvolutionalDimensions == 3 ? EConvMode::D3
+		: /*NumberConvolutionalDimensions > 3 ?*/ EConvMode::DN);
+	if (ConvMode == EConvMode::DN)
+	{
+		Parameters->NumberConvolutionalDimensions = NumberConvolutionalDimensions;
+	}
+	// SRV/UAV variables
+	Parameters->OutputUAV = OutputTensor.GetBufferUAVRef();
+	Parameters->XOrXWithZerosSRV = XOrXWithZeros.GetBufferSRVRef(); // InputTensors[0] if normal convolution
+	Parameters->WSRV = InputTensors[1]->GetBufferSRVRef();
+	// Set output to bias
+	const bool bHasBias = (InputTensors.Num() > 2);
+	if (bHasBias)
+	{
+		Parameters->BSRV = InputTensors[2]->GetBufferSRVRef();
+	}
+	// Set shader
+	FConvBaseCS::FPermutationDomain PermutationVector;
+	PermutationVector.Set<FConvBaseCS::FHasBias>(bHasBias);
+	PermutationVector.Set<FConvBaseCS::FConvMode>(ConvMode);
+	TShaderMapRef<FConvBaseCS> ComputeShader(GetGlobalShaderMap(GMaxRHIFeatureLevel), PermutationVector);
+	// Run shader
+	const uint32 ThreadGroupCountValueX = FMath::DivideAndRoundUp(Parameters->OutputVolume, FConvBaseCS::THREADGROUP_SIZE_X);
+	FComputeShaderUtils::AddPass(
+		*InOutGraphBuilder,
+		RDG_EVENT_NAME("FConvBaseCS() - %dD convolution, Operator: %s", NumberConvolutionalDimensions, *Name),
+		ComputeShader,
+		Parameters,
+		FIntVector(ThreadGroupCountValueX, 1, 1));
+}
+
+
+
+/* FConvBaseOperator private functions
+ *****************************************************************************/
+
+bool FConvBaseOperator::SanityChecksForTensors(const TArray<int64>& InPartiallySetOutputSizes) const
+{
+	// InlinedTensor should be false
+	if (InlinedTensor > -1)
+	{
+		UE_LOG(LogNeuralNetworkInference, Warning, TEXT("FConvBaseOperator::SanityChecksForTensors(): This operator cannot be inlined."));
+		return false;
+	}
+	// Input sanity checks
+	const TArray<FNeuralTensor*>& InputTensors = GetInputTensorsConst();
+	const TArray<int64>& XSizes = InputTensors[0]->GetSizes();
+	const TArray<int64>& WSizes = InputTensors[1]->GetSizes();
+	if (InputTensors.Num() > 2)
+	{
+		if (!bIsTransposed)
+		{
+			if (WSizes[0] != InputTensors[2]->Num())
+			{
+				UE_LOG(LogNeuralNetworkInference, Warning,
+					TEXT("FConvBaseOperator::SanityChecksForTensors(): InputTensors[1].GetSize(0) == InputTensors[2].Num() failed (%d != %d)."),
+					WSizes[0], InputTensors[2]->Num());
+				return false;
+			}
+		}
+		else
+		{
+			if (Group * WSizes[1] != InputTensors[2]->Num())
+			{
+				UE_LOG(LogNeuralNetworkInference, Warning,
+					TEXT("FConvBaseOperator::SanityChecksForTensors(): Group x InputTensors[1].GetSize(1) == InputTensors[2].Num() failed (%d x %d != %d)."),
+					Group, WSizes[1], InputTensors[2]->Num());
+				return false;
+			}
+		}
+	}
+	// Grouped convolution checks
+	if (XSizes[1] % Group != 0 || WSizes[0] % Group != 0)
+	{
+		UE_LOG(LogNeuralNetworkInference, Warning,
+			TEXT("FConvBaseOperator::SanityChecksForTensors(): XSizes[1] %% Group == 0 && WSizes[0] %% Group == 0 failed (%d %% %d = %d && %d %% %d = %d)."),
+			XSizes[1], Group, XSizes[1] % Group,
+			WSizes[0], Group, WSizes[0] % Group);
+		return false;
+	}
+	// Check C matches: XWithZeros = {B, C, ...}, W = {M, C/Group, ...}
+	if (!bIsTransposed)
+	{
+		if (XSizes[1] != Group * WSizes[1])
+		{
+			UE_LOG(LogNeuralNetworkInference, Warning,
+				TEXT("FConvBaseOperator::SanityChecksForTensors(): XSizes[1] == Group * WSizes[1] failed (%d != %d x %d)."),
+				XSizes[1], Group, WSizes[1]);
+			return false;
+		}
+	}
+	else
+	{
+		if (XSizes[1] != WSizes[0])
+		{
+			UE_LOG(LogNeuralNetworkInference, Warning,
+				TEXT("FConvBaseOperator::SanityChecksForTensors(): XSizes[1] == WSizes[0] failed (%d != %d)."),
+				XSizes[1], WSizes[0]);
+			return false;
+		}
+	}
+	// XWithZeros, W, Output compatible to each other
+	if (InPartiallySetOutputSizes.Num() != InputTensors[0]->GetNumberDimensions() || InPartiallySetOutputSizes.Num() != InputTensors[1]->GetNumberDimensions())
+	{
+		UE_LOG(LogNeuralNetworkInference, Warning,
+			TEXT("FConvBaseOperator::SanityChecksForTensors(): InPartiallySetOutputSizes.Num() == InputTensors[0]->GetNumberDimensions() == InputTensors[1]->GetNumberDimensions() failed (%d vs. %d vs. %d)."),
+			InPartiallySetOutputSizes.Num(), InputTensors[0]->GetNumberDimensions(), InputTensors[1]->GetNumberDimensions());
+		return false;
+	}
+	// Check Output - Input sizes matches: XWithZeros = {B, C, ...}, W = {M, C/Group, ...}, Output = {N, M, ...}
+	if (InPartiallySetOutputSizes[0] != XSizes[0])
+	{
+		UE_LOG(LogNeuralNetworkInference, Warning,
+			TEXT("FConvBaseOperator::SanityChecksForTensors(): InPartiallySetOutputSizes[0] == XSizes[0] failed (%d != %d)."),
+			InPartiallySetOutputSizes[0], XSizes[0]);
+		return false;
+	}
+	if (!bIsTransposed)
+	{
+		if (InPartiallySetOutputSizes[1] != WSizes[0])
+		{
+			UE_LOG(LogNeuralNetworkInference, Warning,
+				TEXT("FConvBaseOperator::SanityChecksForTensors(): InPartiallySetOutputSizes[1] == WSizes[0] failed (%d != %d)."),
+				InPartiallySetOutputSizes[1], WSizes[0]);
+			return false;
+		}
+	}
+	else if (InPartiallySetOutputSizes[1] != Group * WSizes[1])
+	{
+		UE_LOG(LogNeuralNetworkInference, Warning,
+			TEXT("FConvBaseOperator::SanityChecksForTensors(): InPartiallySetOutputSizes[1] == Group * WSizes[1] failed (%d != %d x %d)."),
+			InPartiallySetOutputSizes[1], Group, WSizes[1]);
+		return false;
+	}
+	return true;
+}
+
+bool FConvBaseOperator::FillOutputShape(const TArray<int64>& InXSizes)
+{
+	// Only applies for transposed convolution
+	if (bIsTransposed)
+	{
+		const int32 NumberConvolutionalDimensions = InXSizes.Num() - 2;
+		// If AutoPad == Same(Lower/Upper), fill OutputShape
+		if (AutoPad == FConvBaseOperator::EAutoPad::SameLower || AutoPad == FConvBaseOperator::EAutoPad::SameUpper)
+		{
+			// Set OutputShapeTemporary
+			TArray<int64> OutputShapeTemporary;
+			OutputShapeTemporary.SetNumUninitialized(NumberConvolutionalDimensions);
+			for (int64 DimensionIndex = 0; DimensionIndex < NumberConvolutionalDimensions; ++DimensionIndex)
+			{
+				OutputShapeTemporary[DimensionIndex] = InXSizes[DimensionIndex + 2] * (Strides.Num() > DimensionIndex ? Strides[DimensionIndex] : 1);
+			}
+			// Set OutputShape from OutputShapeTemporary
+			if (OutputShape.IsEmpty())
+			{
+				OutputShape = OutputShapeTemporary;
+			}
+			// Sanity check
+			else if (OutputShape != OutputShapeTemporary)
+			{
+				UE_LOG(LogNeuralNetworkInference, Warning, TEXT("FConvTransposeOperator::FillOutputShape(): OutputShape == OutputShapeTemporary failed, %s != %s."),
+					*FNeuralTensor(OutputShape).ToString(/*MaxNumberElementsToDisplay*/-1, /*bReturnOnlyData*/true), *FNeuralTensor(OutputShapeTemporary).ToString(/*MaxNumberElementsToDisplay*/-1, /*bReturnOnlyData*/true));
+				return false;
+			}
+		}
+	}
+	return true;
+}
+
+bool FConvBaseOperator::FillDilationsIfEmpty(const int32 InNumberConvolutionalDimensions)
+{
+	// Create and fill InOutDilations (which might be empty)
+	if (Dilations.Num() < 1)
+	{
+		Dilations.Init(1, InNumberConvolutionalDimensions, /*bShouldCreate64BitVersion*/true);
+	}
+	// Dilations sanity checks
+	if (Dilations.Num() > 0 && Dilations.Num() != InNumberConvolutionalDimensions) // Right dimensions
+	{
+		UE_LOG(LogNeuralNetworkInference, Warning, TEXT("FConvBaseOperator::FillDilationsIfEmpty(): Dilations.Num() == 0 || Dilations.Num() == InNumberConvolutionalDimensions failed (%d != 0 || %d)."),
+			Dilations.Num(), InNumberConvolutionalDimensions);
+		return false;
+	}
+	for (const int64 Dilation : Dilations) // Dilation >= 1 on each dimension
+	{
+		if (Dilation < 1)
+		{
+			UE_LOG(LogNeuralNetworkInference, Warning, TEXT("FConvBaseOperator::FillDilationsIfEmpty(): Dilation (%d) < 1 found! Dilation must be >= 1."), Dilation);
+			return false;
+		}
+	}
+	return true;
+}
+
+TArray<int64> FConvBaseOperator::CreateEffectiveKernelSizes(const int32 InNumberConvolutionalDimensions, const TArray<int64>& InWSizes)
+{
+	// Set EffectiveKernelSizes, whose size with dilation is: k_prime = k + (k-1)(d-1) = d(k-1)+1
+	TArray<int64> EffectiveKernelSizes;
+	EffectiveKernelSizes.SetNumUninitialized(InNumberConvolutionalDimensions);
+	for (int64 DimensionIndex = 0; DimensionIndex < InNumberConvolutionalDimensions; ++DimensionIndex)
+	{
+		EffectiveKernelSizes[DimensionIndex] = Dilations[DimensionIndex] * (InWSizes[DimensionIndex + 2] - 1) + 1;
+	}
+	return EffectiveKernelSizes;
+}
+
+bool FConvBaseOperator::EstimateAndFillPads(const TArray<int64>& InEffectiveKernelSizes, const int32 InNumberDimensions, const TArray<int64>& InXSizes)
+{
+	const int32 NumberConvolutionalDimensions = InNumberDimensions - 2;
+	// Create and fill Pads (which might be empty)
+	// Explicit or no padding is used
+	if (AutoPad == EAutoPad::NotSet || AutoPad == EAutoPad::Valid)
+	{
+		// Set Pads (if empty)
+		if (Pads.IsEmpty())
+		{
+			Pads.Init(0, 2 * NumberConvolutionalDimensions, /*bShouldCreate64BitVersion*/true);
+		}
+		// If transposed conv
+		if (bIsTransposed)
+		{
+			// Pads --> Equivalent non-transposed pads
+			Pads = FPrivateConvBaseOperator::TransposePadsFromConvPads(Pads, AutoPad, InEffectiveKernelSizes);
+		}
+	}
+	// AutoPad == Same(Lower/Upper)
+	// If the padding is an odd number, the extra padding is added at the beginning (EAutoPad::SameLower) or end (EAutoPad::SameUpper)
+	// How to obtain the total padding? o = Floor( (i + pL + pR - InEffectiveKernelSizes) / s) + 1
+	//     - The minimum value of padding is assuming the floor does not have remainder (i.e., removing the floor) --> p_min = (o-1) x s - i + InEffectiveKernelSizes, and we have s total solutions for this equation
+	//     - The maximum value of padding is assuming the floor remainder is s-1 --> p_max = o s - i + InEffectiveKernelSizes - 1
+	//     - The standard seems to be using p_min, so we use that one
+	// output_shape[i] = Ceil(input_shape[i] / strides[i]) for each axis i, padding split between the 2 sides equally
+	else if (AutoPad == EAutoPad::SameLower || AutoPad == EAutoPad::SameUpper)
+	{
+		TArray<int64> PadsArrayTemp;
+		PadsArrayTemp.SetNumUninitialized(2 * NumberConvolutionalDimensions);
+		const TArray<int64>& XOrXWithZerosSizes = (AuxiliaryTensors.IsEmpty() ? InXSizes : AuxiliaryTensors[0]->GetSizes()); // XSizes for conv and Stride = 1, XWithZerosSizes for transposed with Stride > 1
+		for (int64 DimensionIndex = 2; DimensionIndex < InNumberDimensions; ++DimensionIndex)
+		{
+			const int64 OutputSize = (!bIsTransposed ? FMath::CeilToInt(InXSizes[DimensionIndex] / (float)Strides[DimensionIndex - 2]) // Convolution
+				: InXSizes[DimensionIndex] * ConvolutionStridesIfTransposedConvolution[DimensionIndex - 2]); // Transpose convolution
+			const int64 TotalPadding = (OutputSize - 1) * Strides[DimensionIndex - 2] - XOrXWithZerosSizes[DimensionIndex] + InEffectiveKernelSizes[DimensionIndex - 2];
+			const int64 PaddingIndex = 2 * (DimensionIndex - 2);
+			PadsArrayTemp[PaddingIndex] = TotalPadding / 2 + (AutoPad == EAutoPad::SameLower ? TotalPadding % 2 : 0);
+			PadsArrayTemp[PaddingIndex + 1] = TotalPadding - PadsArrayTemp[PaddingIndex];
+		}
+		Pads.Move(PadsArrayTemp);
+	}
+	else
+	{
+		UE_LOG(LogNeuralNetworkInference, Warning, TEXT("FConvBaseOperator::EstimateAndFillPads(): Unknown AutoPad = %d."), (int32)AutoPad);
+		return false;
+	}
+	return true;
+}
+
+bool FConvBaseOperator::FillOutputSizesAndPaddedMarginsAndSanityChecks(TArray<int64>& InOutPartiallySetOutputSizes, const TArray<int64>& InEffectiveKernelSizes)
+{
+	const TArray<FNeuralTensor*>& InputTensors = GetInputTensorsConst();
+	const FNeuralTensor& XOrXWithZeros = *(AuxiliaryTensors.Num() == 0 ? InputTensors[0] : AuxiliaryTensors[0]);
+	const TArray<int64>& XOrXWithZerosSizes = XOrXWithZeros.GetSizes();
+	const TArray<int64>& WSizes = InputTensors[1]->GetSizes();
+	const int32 NumberDimensions = XOrXWithZeros.GetNumberDimensions(); // InputTensors[0] if normal convolution
+	const int32 NumberConvolutionalDimensions = NumberDimensions - 2;
+	// XWithZeros.GetSize(Index>1) >= Dilations*(W.GetSize(Index>1)-1) + 1
+	// I.e., each dimension of input XWithZeros must be greater than each dimension of the weights W to avoid segmentation fault (given current implementation)
+	for (int64 DimensionIndex = 2; DimensionIndex < NumberDimensions; ++DimensionIndex)
+	{
+		const int64 PaddingIndex = 2 * (DimensionIndex - 2);
+		if (XOrXWithZerosSizes[DimensionIndex] + Pads[PaddingIndex] + Pads[PaddingIndex + 1] < InEffectiveKernelSizes[DimensionIndex - 2])
+		{
+			UE_LOG(LogNeuralNetworkInference, Warning,
+				TEXT("FConvBaseOperator::FillOutputSizesAndPaddedMarginsAndSanityChecks(): XOrXWithZerosSizes[%d] >= InEffectiveKernelSizes[%d] failed (%d vs. %d)."
+					" I.e., each dimension of input XOrXXWithZeros must be greater than each dimension of the effective kernel/weight size W (taking the dilation into account)."),
+				DimensionIndex, DimensionIndex - 2, XOrXWithZerosSizes[DimensionIndex], InEffectiveKernelSizes[DimensionIndex - 2]);
+			return false;
+		}
+	}
+	// Set InOutPartiallySetOutputSizes based on XOrXWithZeros, W, and padding
+	for (int64 DimensionIndex = 2; DimensionIndex < NumberDimensions; ++DimensionIndex)
+	{
+		const int64 ActualDimension = DimensionIndex - 2;
+		const int64 PaddingIndex = 2 * ActualDimension;
+		const int64 ExpectedOutputSize = /*Implicit floor*/((XOrXWithZerosSizes[DimensionIndex] + Pads[PaddingIndex] + Pads[PaddingIndex + 1] - InEffectiveKernelSizes[DimensionIndex - 2]) / Strides[ActualDimension]) + 1;
+		const int64 OutputPad = (OutputPadding.Num() > 0 ? OutputPadding[ActualDimension] : 0);
+		// For Conv a = 0, but for ConvTranspose: o' = s(i' - 1) + a + k - 2p, where a = #zeros added at the end of XOrXWithZerosSizes
+		InOutPartiallySetOutputSizes[DimensionIndex] = (OutputShape.Num() > 0 ? OutputShape[ActualDimension] : ExpectedOutputSize + OutputPad);
+	}
+	const TArray<int64>& OutputSizes = InOutPartiallySetOutputSizes; // At this point, InOutPartiallySetOutputSizesis finally OutputSizes
+	// Set OutputPaddedMargins
+	// Sides limits deduction
+	// - Formula: i_i = d w_i + o_i s - pL
+	// - Left side: i_i >= 0? That means w_i = 0 (because 1,2,...,k are gonna give a higher result) --> 0 <= 0 + o_i s - pL --> o_i s >= pL --> i_i >= 0 if o_i >= pL/s --> To account for pL % s != 0
+	//     --> i_i >= 0 if o_i >= Ceil(pL/s) >= pL/s   and   i_i might be < 0 if o_i < Ceil(pL/s)
+	// - Right side: i_i <= i-1? That means w_i = k-1 --> i-1 >= d(k-1) + o_i s - pL --> o_i s <= i-1 - d(k-1) + pL --> o_i <= [i - 1 - d(k-1) + pL]/s
+	//     --> i_i <= i-1 if o_i <= Floor([i - 1 - d(k-1) + pL]/s) <= [i - 1 - d(k-1) + pL]/s   and   i_i might be > i-1 if o_i > Floor([i - 1 + pL - d(k-1)]/s)
+	// Particular cases:
+	// - If Pads[i] = 0 --> OutputPaddedMargins = [0, XSize]
+	// - If Strides[i] = 1 --> OutputPaddedMargins = [PadsLeft[i], OutputSize-PadsRight[i]]
+	OutputPaddedMargins.SetNumUninitialized(Pads.Num());
+	for (int64 DimensionIndex = 0; DimensionIndex < NumberConvolutionalDimensions; ++DimensionIndex)
+	{
+		const int64 PaddingIndex = 2 * DimensionIndex;
+		OutputPaddedMargins[PaddingIndex] = FMath::CeilToInt(Pads[PaddingIndex] / (float)Strides[DimensionIndex]);
+		const int64 OutputPaddedMarginsRight = 1 + FMath::FloorToInt( // +1 because of 0 index
+			(XOrXWithZerosSizes[DimensionIndex + 2] - 1 - Dilations[DimensionIndex] * (WSizes[DimensionIndex + 2] - 1) + Pads[PaddingIndex]) / (float)Strides[DimensionIndex]);
+		// Problem: Because we are using floor and ceil, it might happen that OutputPaddedMargins[PaddingIndex + 1] < OutputPaddedMargins[PaddingIndex], which
+		// would result in overriding the same value twice, generating wrong results!
+		// Fix: OutputPaddedMargins[PaddingIndex + 1] = Max(both OutputPaddedMargins) to avoid applying the convolution twice over the same pixel. This way,
+		// CPU_CONVOLUTION_XD() is not executed and CPU_CONVOLUTION_XD_PADDED() processes all pixels (because they would all need padding).
+		OutputPaddedMargins[PaddingIndex + 1] = FMath::Max(OutputPaddedMargins[PaddingIndex], OutputPaddedMarginsRight);
+		// Self sanity checks
+		// If s==1 || Pads[i] == 0 --> OutputPaddedMargins = [0, OutputSizes]
+		const int64 OutputShapeIndex = (OutputShape.Num() > 0 ? OutputShape[DimensionIndex] : 0);
+		if (OutputShapeIndex == 0 && (Strides[DimensionIndex] == 1 || (Pads[PaddingIndex] == 0 && Pads[PaddingIndex + 1] == 0)))
+		{
+			const int64 OutputPad = (OutputPadding.Num() > 0 ? OutputPadding[DimensionIndex] : 0);
+			if (OutputPaddedMargins[PaddingIndex] != Pads[PaddingIndex] || OutputPaddedMarginsRight != OutputSizes[DimensionIndex + 2] - Pads[PaddingIndex + 1] - OutputPad)
+			{
+				UE_LOG(LogNeuralNetworkInference, Warning,
+					TEXT("FConvBaseOperator::FillOutputSizesAndPaddedMarginsAndSanityChecks(): Even though s=1 and/or Pads[i]=0, OutputPaddedMargins[%d,%d] != [Pads[%d], Output[%d]-Pads[%d]-OutputPad],"
+						" i.e., [%d,%d] != [%d,%d-%d-%d], where the final OutputPaddedMarginsRight = %d."),
+					PaddingIndex, PaddingIndex + 1, PaddingIndex, DimensionIndex + 2, PaddingIndex + 1,
+					OutputPaddedMargins[PaddingIndex], OutputPaddedMarginsRight, Pads[PaddingIndex], OutputSizes[DimensionIndex + 2], Pads[PaddingIndex + 1], OutputPad, OutputPaddedMargins[PaddingIndex + 1]);
+				return false;
+			}
+		}
+		// If s>1 && Pads[i] > 0 --> Make sure OutputPaddedMargins[i] >= 0 && OutputPaddedMargins[i+1] < OutputSizes
+		else
+		{
+			if (OutputPaddedMargins[PaddingIndex] < 0 || OutputPaddedMargins[PaddingIndex + 1] > OutputSizes[DimensionIndex + 2])
+			{
+				UE_LOG(LogNeuralNetworkInference, Warning,
+					TEXT("FConvBaseOperator::FillOutputSizesAndPaddedMarginsAndSanityChecks(): OutputPaddedMargins[%d]=%d < 0 || OutputPaddedMargins[%d]=%d > OutputSizes[%d]=%d. Other info: s=%d, Pads[%d,%d+1]=[%d,%d]."),
+					PaddingIndex, OutputPaddedMargins[PaddingIndex], PaddingIndex + 1, OutputPaddedMargins[PaddingIndex + 1], DimensionIndex + 2, OutputSizes[DimensionIndex + 2],
+					Strides[DimensionIndex], PaddingIndex, PaddingIndex + 1, Pads[PaddingIndex], Pads[PaddingIndex + 1]);
+				return false;
+			}
+		}
+	}
+	// Initialize OutputTensor
+	FNeuralTensor& OutputTensor = GetOutputTensorNoConst();
+	OutputTensor.SetNumUninitialized(OutputSizes, XOrXWithZeros.GetDataType());
+	// Sanity check
+	if (!FNeuralNetworkInferenceUtils::SizeSanityChecks(GetOutputTensorsConst(), 1, 1, 3, -1))
+	{
+		return false;
+	}
+	return true;
+}
+
+#define MULTICHANNEL_CPU_CONVOLUTION_FOR_LOOPS_INIT() \
+	/* Per batch N */ \
+	for (int64 BatchIndex = 0; BatchIndex < OutputSizes[0]; ++BatchIndex) \
+	{ \
+		/* Per output channel M */ \
+		for (int64 OutputChannelIndex = 0; OutputChannelIndex < OutputSizes[1]; ++OutputChannelIndex) \
+		{ \
+			const uint64 YPtrBatchChannelIndex = BatchIndex * OutputBatchVolume + OutputChannelIndex * OutputImageArea; \
+			/* Per input channel C */ \
+			for (int64 InputChannelIndex = 0; InputChannelIndex < !bIsTransposed * WSizes[1] + bIsTransposed * WSizes[0] / Group; ++InputChannelIndex) /* WSizes[1] and not XOrXWithZerosSizes[1] for grouped convolution */ \
+			{ \
+				const uint64 GroupedInputChannelIndex = InputChannelIndex + bIsTransposed * (OutputChannelIndex / MIntoGroup * CIntoGroup); \
+				const uint64 WBatchChannelIndex = !bIsTransposed * (OutputChannelIndex * WBatchVolume + InputChannelIndex * WImageArea) + bIsTransposed * (InputChannelIndex * WBatchVolume + OutputChannelIndex * WImageArea); \
+				const uint64 XBatchChannelIndex = BatchIndex * XOrXWithZerosBatchVolume + GroupedInputChannelIndex * XOrXWithZerosImageArea \
+					+  /* For grouped convolution */ !bIsTransposed * (OutputChannelIndex/MIntoGroup * WSizes[1]);
+#define CPU_CONVOLUTION_FOR_LOOP_X1_INIT(ForX1Begin, ForX1End) \
+				/* Convolution - D1 */ \
+				for (int64 OutputD1Index = ForX1Begin; OutputD1Index < ForX1End; ++OutputD1Index) /* OutputD1Index = 0 : OutputSizes[2] if no padding */ \
+				{
+#define CPU_CONVOLUTION_FOR_LOOP_X2_INIT(ForX2Begin, ForX2End) \
+					/* Convolution - D2 */ \
+					for (int64 OutputD2Index = ForX2Begin; OutputD2Index < ForX2End; ++OutputD2Index) /* OutputD2Index = 0 : OutputSizes[3] if no padding */ \
+					{
+#define CPU_CONVOLUTION_FOR_LOOP_X3_INIT(ForX3Begin, ForX3End) \
+						/* Convolution - D3 */ \
+						for (int64 OutputD3Index = ForX3Begin; OutputD3Index < ForX3End; ++OutputD3Index) /* OutputD3Index = 0 : OutputSizes[4] if no padding */ \
+						{
+#define CPU_CONVOLUTION_FOR_LOOP_W1_INIT() \
+							/* Convolution - D1k */ \
+							for (int64 W1Index = 0; W1Index < WSizes[2]; ++W1Index) \
+							{ \
+								const int64 X1Index = (Dilations[0] * W1Index + OutputD1Index * Strides[0] - Pads[0]); /* Pads[i] = 0 if no padding */
+#define CPU_CONVOLUTION_FOR_LOOP_W2_INIT() \
+								/* Convolution - D2k */ \
+								for (int64 W2Index = 0; W2Index < WSizes[3]; ++W2Index) \
+								{ \
+									const int64 X2Index = (Dilations[1] * W2Index + OutputD2Index * Strides[1] - Pads[2]); /* Pads[i] = 0 if no padding */
+#define CPU_CONVOLUTION_FOR_LOOP_W3_INIT() \
+									/* Convolution - D3k */ \
+									for (int64 W3Index = 0; W3Index < WSizes[4]; ++W3Index) \
+									{ \
+										const int64 X3Index = (Dilations[2] * W3Index + OutputD3Index * Strides[2] - Pads[4]); /* Pads[i] = 0 if no padding */
+
+#define CPU_CONVOLUTION_FOR_LOOP_W3_END() \
+									}
+#define CPU_CONVOLUTION_FOR_LOOP_W2_END() \
+								}
+#define CPU_CONVOLUTION_FOR_LOOP_W1_END() \
+							}
+#define CPU_CONVOLUTION_FOR_LOOP_X3_END() \
+						}
+#define CPU_CONVOLUTION_FOR_LOOP_X2_END() \
+					}
+#define CPU_CONVOLUTION_FOR_LOOP_X1_END() \
+				}
+#define MULTICHANNEL_CPU_CONVOLUTION_FOR_LOOPS_END() \
+			} \
+		} \
+	}
+
+#define CPU_CONVOLUTION_1D_FOR_LOOPS_INIT(ForX1Begin, ForX1End) \
+	CPU_CONVOLUTION_FOR_LOOP_X1_INIT(ForX1Begin, ForX1End) \
+		const uint64 YValueIndex = YPtrBatchChannelIndex + OutputD1Index; \
+		CPU_CONVOLUTION_FOR_LOOP_W1_INIT()
+
+#define CPU_CONVOLUTION_1D_FOR_LOOPS_END() \
+		CPU_CONVOLUTION_FOR_LOOP_W1_END() \
+	CPU_CONVOLUTION_FOR_LOOP_X1_END()
+
+#define CPU_CONVOLUTION_1D(ForX1Begin, ForX1End) \
+	CPU_CONVOLUTION_1D_FOR_LOOPS_INIT(ForX1Begin, ForX1End) \
+		OutY.At<float>(YValueIndex) += InW.At<float>(WBatchChannelIndex + W1Index) * InXOrXWithZerosSizes.At<float>(XBatchChannelIndex + X1Index); \
+	CPU_CONVOLUTION_1D_FOR_LOOPS_END()
+
+#define CPU_CONVOLUTION_1D_PADDED(ForX1Begin, ForX1End) \
+	CPU_CONVOLUTION_1D_FOR_LOOPS_INIT(ForX1Begin, ForX1End) \
+		if (X1Index > -1 && X1Index < XOrXWithZerosSizes[2]) /* Always true if no padding */ \
+		{ \
+			OutY.At<float>(YValueIndex) += InW.At<float>(WBatchChannelIndex + W1Index) * InXOrXWithZerosSizes.At<float>(XBatchChannelIndex + X1Index); \
+		} \
+	CPU_CONVOLUTION_1D_FOR_LOOPS_END()
+
+#define CPU_CONVOLUTION_2D_FOR_LOOPS_INIT(ForX1Begin, ForX1End, ForX2Begin, ForX2End) \
+	CPU_CONVOLUTION_FOR_LOOP_X1_INIT(ForX1Begin, ForX1End) \
+		CPU_CONVOLUTION_FOR_LOOP_X2_INIT(ForX2Begin, ForX2End) \
+			const uint64 YValueIndex = YPtrBatchChannelIndex + OutputD1Index * OutputSizes[3] + OutputD2Index; \
+			CPU_CONVOLUTION_FOR_LOOP_W1_INIT() \
+				CPU_CONVOLUTION_FOR_LOOP_W2_INIT()
+
+#define CPU_CONVOLUTION_2D_FOR_LOOPS_END() \
+					CPU_CONVOLUTION_FOR_LOOP_W2_END() \
+				CPU_CONVOLUTION_FOR_LOOP_W1_END() \
+		CPU_CONVOLUTION_FOR_LOOP_X2_END() \
+	CPU_CONVOLUTION_FOR_LOOP_X1_END()
+
+#define CPU_CONVOLUTION_2D(ForX1Begin, ForX1End, ForX2Begin, ForX2End) \
+	CPU_CONVOLUTION_2D_FOR_LOOPS_INIT(ForX1Begin, ForX1End, ForX2Begin, ForX2End) \
+		const int64 WIndex = W1Index * WSizes[3] + W2Index; \
+		const int64 XIndex = X1Index * XOrXWithZerosSizes[3] + X2Index; \
+		OutY.At<float>(YValueIndex) += InW.At<float>(WBatchChannelIndex + WIndex) * InXOrXWithZerosSizes.At<float>(XBatchChannelIndex + XIndex); \
+	CPU_CONVOLUTION_2D_FOR_LOOPS_END()
+
+#define CPU_CONVOLUTION_2D_PADDED(ForX1Begin, ForX1End, ForX2Begin, ForX2End) \
+	CPU_CONVOLUTION_2D_FOR_LOOPS_INIT(ForX1Begin, ForX1End, ForX2Begin, ForX2End) \
+		if (X1Index > -1 && X1Index < XOrXWithZerosSizes[2] && X2Index > -1 && X2Index < XOrXWithZerosSizes[3]) /* Always true if no padding */ \
+		{ \
+			const int64 WIndex = W1Index * WSizes[3] + W2Index; \
+			const int64 XIndex = X1Index * XOrXWithZerosSizes[3] + X2Index; \
+			OutY.At<float>(YValueIndex) += InW.At<float>(WBatchChannelIndex + WIndex) * InXOrXWithZerosSizes.At<float>(XBatchChannelIndex + XIndex); \
+		} \
+	CPU_CONVOLUTION_2D_FOR_LOOPS_END()
+
+#define CPU_CONVOLUTION_3D_FOR_LOOPS_INIT(ForX1Begin, ForX1End, ForX2Begin, ForX2End, ForX3Begin, ForX3End) \
+	CPU_CONVOLUTION_FOR_LOOP_X1_INIT(ForX1Begin, ForX1End) \
+		CPU_CONVOLUTION_FOR_LOOP_X2_INIT(ForX2Begin, ForX2End) \
+			CPU_CONVOLUTION_FOR_LOOP_X3_INIT(ForX3Begin, ForX3End) \
+				const uint64 YValueIndex = YPtrBatchChannelIndex + (OutputD1Index * OutputSizes[3] + OutputD2Index) * OutputSizes[4] + OutputD3Index; \
+				CPU_CONVOLUTION_FOR_LOOP_W1_INIT() \
+					CPU_CONVOLUTION_FOR_LOOP_W2_INIT() \
+						CPU_CONVOLUTION_FOR_LOOP_W3_INIT()
+
+#define CPU_CONVOLUTION_3D_FOR_LOOPS_END() \
+						CPU_CONVOLUTION_FOR_LOOP_W3_END() \
+					CPU_CONVOLUTION_FOR_LOOP_W2_END() \
+				CPU_CONVOLUTION_FOR_LOOP_W1_END() \
+			CPU_CONVOLUTION_FOR_LOOP_X3_END() \
+		CPU_CONVOLUTION_FOR_LOOP_X2_END() \
+	CPU_CONVOLUTION_FOR_LOOP_X1_END()
+
+#define CPU_CONVOLUTION_3D(ForX1Begin, ForX1End, ForX2Begin, ForX2End, ForX3Begin, ForX3End) \
+	CPU_CONVOLUTION_3D_FOR_LOOPS_INIT(ForX1Begin, ForX1End, ForX2Begin, ForX2End, ForX3Begin, ForX3End) \
+		const int64 WIndex = (W1Index * WSizes[3] + W2Index) * WSizes[4] + W3Index; \
+		const int64 XIndex = (X1Index * XOrXWithZerosSizes[3] + X2Index) * XOrXWithZerosSizes[4] + X3Index; /* Pads[i] = 0 if no padding */ \
+		OutY.At<float>(YValueIndex) += InW.At<float>(WBatchChannelIndex + WIndex) * InXOrXWithZerosSizes.At<float>(XBatchChannelIndex + XIndex); \
+	CPU_CONVOLUTION_3D_FOR_LOOPS_END()
+
+#define CPU_CONVOLUTION_3D_PADDED(ForX1Begin, ForX1End, ForX2Begin, ForX2End, ForX3Begin, ForX3End) \
+	CPU_CONVOLUTION_3D_FOR_LOOPS_INIT(ForX1Begin, ForX1End, ForX2Begin, ForX2End, ForX3Begin, ForX3End) \
+		if (X1Index > -1 && X1Index < XOrXWithZerosSizes[2] && X2Index > -1 && X2Index < XOrXWithZerosSizes[3] && X3Index > -1 && X3Index < XOrXWithZerosSizes[4]) /* Always true if no padding */ \
+		{ \
+			const int64 WIndex = (W1Index * WSizes[3] + W2Index) * WSizes[4] + W3Index; \
+			const int64 XIndex = (X1Index * XOrXWithZerosSizes[3] + X2Index) * XOrXWithZerosSizes[4] + X3Index; /* Pads[i] = 0 if no padding */ \
+			OutY.At<float>(YValueIndex) += InW.At<float>(WBatchChannelIndex + WIndex) * InXOrXWithZerosSizes.At<float>(XBatchChannelIndex + XIndex); \
+		} \
+	CPU_CONVOLUTION_3D_FOR_LOOPS_END()
+
+void FConvBaseOperator::MultiChannelConvolutionCPU(FNeuralTensor& OutY, const FNeuralTensor& InXOrXWithZerosSizes, const FNeuralTensor& InW)
+{
+	// Convection
+	// XOrXWithZeros = {N, C, Hx, Wx}
+	// W = {M, C/Group, Hk, Wk}
+	// Y = {N, M, Hy, Wy}
+	const TArray<int64>& XOrXWithZerosSizes = InXOrXWithZerosSizes.GetSizes();
+	const TArray<int64>& WSizes = InW.GetSizes();
+	const TArray<int64>& OutputSizes = OutY.GetSizes();
+	const int64 WBatchVolume = InW.Num() / WSizes[0];
+	const int64 WImageArea = WBatchVolume / WSizes[1];
+	const int64 XOrXWithZerosBatchVolume = InXOrXWithZerosSizes.Num() / XOrXWithZerosSizes[0];
+	const int64 XOrXWithZerosImageArea = XOrXWithZerosBatchVolume / XOrXWithZerosSizes[1];
+	const int64 OutputBatchVolume = OutY.Num() / OutputSizes[0];
+	const int64 OutputImageArea = OutputBatchVolume / OutputSizes[1];
+	const int32 NumberDimensions = OutY.GetNumberDimensions();
+	const int64 MIntoGroup = (!bIsTransposed ? WSizes[0] / Group : WSizes[1]);
+	const int64 CIntoGroup = WSizes[0] / Group;
+	// 1D convolution
+	// We make 1D, 2D, and 3D simple and efficient enough (there will be cache misses and lots of if statements for the padded area)
+	// 1. A for loop for the middle part, which does not require any padding and will have almost no cache misses
+	// 2. Different for loops with if's for all the sides (top, bottom, for 2D: left without top/bottom, right without top/bottom, etc.), which require padding
+	if (NumberDimensions == 3)
+	{
+		// Per batch N, output channel M, input channel C
+		MULTICHANNEL_CPU_CONVOLUTION_FOR_LOOPS_INIT()
+			CPU_CONVOLUTION_1D_PADDED(	/*ForX1Begin*/0,						/*ForX1End*/OutputPaddedMargins[0])	// 1D convolution - Left padding
+			CPU_CONVOLUTION_1D(			/*ForX1Begin*/OutputPaddedMargins[0],	/*ForX1End*/OutputPaddedMargins[1])	// 1D convolution
+			CPU_CONVOLUTION_1D_PADDED(	/*ForX1Begin*/OutputPaddedMargins[1],	/*ForX1End*/OutputSizes[2])			// 1D convolution - Right padding
+			//CPU_CONVOLUTION_1D_PADDED(	/*ForX1Begin*/0,						/*ForX1End*/OutputSizes[2])			// Less efficient equivalent - Full 1D convolution
+		MULTICHANNEL_CPU_CONVOLUTION_FOR_LOOPS_END()
+	}
+	// >= 2D convolution
+	else
+	{
+		// 2D convolution - Same than 1D convolution but adding the 2nd dimensional padding
+		if (NumberDimensions == 4)
+		{
+			// Per batch N, output channel M, input channel C
+			MULTICHANNEL_CPU_CONVOLUTION_FOR_LOOPS_INIT()
+				// 2D convolution - Top padding
+				CPU_CONVOLUTION_2D_PADDED(/*D1*/0, OutputPaddedMargins[0],						/*D2*/0, OutputSizes[3]);
+				// 2D convolution - Left padding (without top/bottom)
+				CPU_CONVOLUTION_2D_PADDED(/*D1*/OutputPaddedMargins[0], OutputPaddedMargins[1],	/*D2*/0, OutputPaddedMargins[2]);
+				// 2D convolution - Center (not padded part)
+				CPU_CONVOLUTION_2D		 (/*D1*/OutputPaddedMargins[0], OutputPaddedMargins[1],	/*D2*/OutputPaddedMargins[2], OutputPaddedMargins[3]);
+				// 2D convolution - Right padding (without top/bottom)
+				CPU_CONVOLUTION_2D_PADDED(/*D1*/OutputPaddedMargins[0], OutputPaddedMargins[1],	/*D2*/OutputPaddedMargins[3], OutputSizes[3]);
+				// 2D convolution - Bottom padding
+				CPU_CONVOLUTION_2D_PADDED(/*D1*/OutputPaddedMargins[1], OutputSizes[2],			/*D2*/0, OutputSizes[3]);
+				//// Less efficient equivalent - Full 2D convolution
+				//CPU_CONVOLUTION_2D_PADDED(/*D1*/0, OutputSizes[2],								/*D2*/0, OutputSizes[3]);
+			MULTICHANNEL_CPU_CONVOLUTION_FOR_LOOPS_END()
+		}
+		// >= 3D convolution
+		else
+		{
+			// 3D convolution - Same than 2D convolution but adding the 3rd dimensional padding
+			if (NumberDimensions == 5)
+			{
+				// Per batch N, output channel M, input channel C
+				MULTICHANNEL_CPU_CONVOLUTION_FOR_LOOPS_INIT()
+					// 3D convolution - Top padding
+					CPU_CONVOLUTION_3D_PADDED(/*D1*/0, OutputPaddedMargins[0],						/*D2*/0, OutputSizes[3],								/*D3*/0, OutputSizes[4]);
+					// 3D convolution - Left padding (without top/bottom)
+					CPU_CONVOLUTION_3D_PADDED(/*D1*/OutputPaddedMargins[0], OutputPaddedMargins[1],	/*D2*/0, OutputPaddedMargins[2],						/*D3*/0, OutputSizes[4]);
+					// 3D convolution - Depth padding (without top/bottom or left/right)
+					CPU_CONVOLUTION_3D_PADDED(/*D1*/OutputPaddedMargins[0], OutputPaddedMargins[1],	/*D2*/OutputPaddedMargins[2], OutputPaddedMargins[3],	/*D3*/0, OutputPaddedMargins[4]);
+					// 3D convolution - Center (not padded part)
+					CPU_CONVOLUTION_3D		 (/*D1*/OutputPaddedMargins[0], OutputPaddedMargins[1],	/*D2*/OutputPaddedMargins[2], OutputPaddedMargins[3],	/*D3*/OutputPaddedMargins[4], OutputPaddedMargins[5]);
+					// 3D convolution - Depth padding (without top/bottom or left/right)
+					CPU_CONVOLUTION_3D_PADDED(/*D1*/OutputPaddedMargins[0], OutputPaddedMargins[1],	/*D2*/OutputPaddedMargins[2], OutputPaddedMargins[3],	/*D3*/OutputPaddedMargins[5], OutputSizes[4]);
+					// 3D convolution - Right padding (without top/bottom)
+					CPU_CONVOLUTION_3D_PADDED(/*D1*/OutputPaddedMargins[0], OutputPaddedMargins[1],	/*D2*/OutputPaddedMargins[3], OutputSizes[3],			/*D3*/0, OutputSizes[4]);
+					// 3D convolution - Bottom padding
+					CPU_CONVOLUTION_3D_PADDED(/*D1*/OutputPaddedMargins[1], OutputSizes[2],			/*D2*/0, OutputSizes[3],								/*D3*/0, OutputSizes[4]);
+					//// Less efficient equivalent - Full 3D convolution
+					//CPU_CONVOLUTION_3D_PADDED(/*D1*/0, OutputSizes[2],								/*D2*/0, OutputSizes[3],								/*D3*/0, OutputSizes[4]);
+				MULTICHANNEL_CPU_CONVOLUTION_FOR_LOOPS_END()
+			}
+			// nD convolution
+			// Very inefficient, it is just meant as a backup if the user do a non expected convolution (i.e., not 1D, 2D, nor 3D).
+			// If performance is required, re-implement 2D/3D idea for the desired convolution size.
+			else
+			{
+				const int32 NumberConvolutionalDimensions = NumberDimensions - 2;
+				TArray<int64> OutputImageAreaIndexes;
+				TArray<int64> WImageAreaIndexes;
+				// Per batch N, output channel M, input channel C
+				MULTICHANNEL_CPU_CONVOLUTION_FOR_LOOPS_INIT()
+					// nD convolution - D1, D2, ..., Dn
+					OutputImageAreaIndexes.Init(0, NumberConvolutionalDimensions);
+					for (int64 OutputImageAreaIndex = 0; OutputImageAreaIndex < OutputImageArea; ++OutputImageAreaIndex)
+					{
+						// Get output index
+						const uint64 YValueIndex = YPtrBatchChannelIndex + OutputImageAreaIndex;
+						float YValue = OutY.At<float>(YValueIndex);
+						WImageAreaIndexes.Init(0, NumberConvolutionalDimensions);
+						// nD convolution - D1k, D2k, ..., Dnk
+						for (int64 WImageAreaIndex = 0; WImageAreaIndex < WImageArea; ++WImageAreaIndex)
+						{
+							// Get XIndex
+							int64 XIndex = (Dilations[0] * WImageAreaIndexes[0] + OutputImageAreaIndexes[0] * Strides[0] - Pads[0]); // Pads[i] = 0 if no padding
+							if (XIndex > -1 && XIndex < XOrXWithZerosSizes[2]) // Always true if no padding
+							{
+								bool bIsPaddedRegionSoItShouldSkipPixel = false;
+								for (int64 XOrXWithZerosImageAreaIndex = 0; XOrXWithZerosImageAreaIndex < NumberConvolutionalDimensions - 1; ++XOrXWithZerosImageAreaIndex)
+								{
+									// For simplicity of the following formulas, ignoring Pads, Dilations, and Strides
+									// XOrXWithZerosImageAreaIndexes[Index] = (OutputImageAreaIndexes[Index] + Dilations[Index] * WImageAreaIndexes[Index])
+									// XIndex1D =   XOrXWithZerosImageAreaIndexes[0]
+									// XIndex2D =   XOrXWithZerosImageAreaIndexes[0] * XOrXWithZerosSizes[3] + XOrXWithZerosImageAreaIndexes[1]
+									// XIndex3D =  (XOrXWithZerosImageAreaIndexes[0] * XOrXWithZerosSizes[3] + XOrXWithZerosImageAreaIndexes[1]) * XOrXWithZerosSizes[4] + XOrXWithZerosImageAreaIndexes[2]
+									// XIndex4D = ((XOrXWithZerosImageAreaIndexes[0] * XOrXWithZerosSizes[3] + XOrXWithZerosImageAreaIndexes[1]) * XOrXWithZerosSizes[4] + XOrXWithZerosImageAreaIndexes[2]) * XOrXWithZerosSizes[5] + XOrXWithZerosImageAreaIndexes[3]
+									const int64 DimensionIndex = 1 + XOrXWithZerosImageAreaIndex;
+									const int64 DimensionValue = (Dilations[DimensionIndex] * WImageAreaIndexes[DimensionIndex] + OutputImageAreaIndexes[DimensionIndex] * Strides[DimensionIndex]
+										- Pads[2 * DimensionIndex]); // Pads[i] = 0 if no padding
+									if (DimensionValue > -1 && DimensionValue < XOrXWithZerosSizes[3 + XOrXWithZerosImageAreaIndex]) // Always true if no padding
+									{
+										XIndex = XIndex * XOrXWithZerosSizes[3 + XOrXWithZerosImageAreaIndex] + DimensionValue;
+									}
+									else // Always false if no padding
+									{
+										bIsPaddedRegionSoItShouldSkipPixel = true;
+										break;
+									}
+								}
+								// Update output
+								if (!bIsPaddedRegionSoItShouldSkipPixel) // Always true if no padding
+								{
+									YValue += InW.At<float>(WBatchChannelIndex + WImageAreaIndex) * InXOrXWithZerosSizes.At<float>(XBatchChannelIndex + XIndex);
+								}
+							}
+							// Update WeightImageAreaIndexes
+							FPrivateConvBaseOperator::NDTensorIndexesPlus1(WImageAreaIndexes, WSizes);
+						}
+						// Update OutputImageAreaIndexes TArray
+						FPrivateConvBaseOperator::NDTensorIndexesPlus1(OutputImageAreaIndexes, OutputSizes);
+						// Update OutY
+						OutY.At<float>(YValueIndex) = YValue;
+					}
+				MULTICHANNEL_CPU_CONVOLUTION_FOR_LOOPS_END()
+			}
+		}
+	}
+}
diff --git a/Engine/Plugins/Experimental/NeuralNetworkInference/Source/NeuralNetworkInference/Private/NeuralOperators/ConvOperator.cpp b/Engine/Plugins/Experimental/NeuralNetworkInference/Source/NeuralNetworkInference/Private/NeuralOperators/ConvOperator.cpp
new file mode 100644
--- /dev/null
+++ b/Engine/Plugins/Experimental/NeuralNetworkInference/Source/NeuralNetworkInference/Private/NeuralOperators/ConvOperator.cpp
@@ -0,0 +1,108 @@
+// Copyright Epic Games, Inc. All Rights Reserved.
+
+#include "NeuralOperators/ConvOperator.h"
+#include "NeuralNetworkInferenceUtils.h"
+#include "NeuralNetworkInferenceUtilsGPU.h"
+
+
+
+/* FPrivateConvOperator auxiliary functions
+ *****************************************************************************/
+
+class FPrivateConvOperator
+{
+public:
+	static FConvOperator NodeProtoToConvOperator(const FNodeProto& InNodeProto);
+};
+
+FConvOperator FPrivateConvOperator::NodeProtoToConvOperator(const FNodeProto& InNodeProto)
+{
+	FConvBaseOperator::EAutoPad AutoPad = FConvBaseOperator::EAutoPad::NotSet;
+	if (const FAttributeProto* EpsilonAttribute = FModelProto::FindElementInArray(TEXT("AutoPad"), InNodeProto.Attribute, /*bMustValueBeFound*/false))
+	{
+		if (EpsilonAttribute->S == TEXT("NOTSET"))
+		{
+			AutoPad = FConvBaseOperator::EAutoPad::NotSet;
+		}
+		else if (EpsilonAttribute->S == TEXT("SAME_UPPER"))
+		{
+			AutoPad = FConvBaseOperator::EAutoPad::SameUpper;
+		}
+		else if (EpsilonAttribute->S == TEXT("SAME_LOWER"))
+		{
+			AutoPad = FConvBaseOperator::EAutoPad::SameLower;
+		}
+		else if (EpsilonAttribute->S == TEXT("VALID"))
+		{
+			AutoPad = FConvBaseOperator::EAutoPad::Valid;
+		}
+		else
+		{
+			UE_LOG(LogNeuralNetworkInference, Warning, TEXT("FConvOperator(): Unknown EpsilonAttribute->S = %s."), *EpsilonAttribute->S);
+		}
+	}
+	TArray<int64> Dilations;
+	if (const FAttributeProto* MomentumAttribute = FModelProto::FindElementInArray(TEXT("Dilations"), InNodeProto.Attribute, /*bMustValueBeFound*/false))
+	{
+		Dilations = MomentumAttribute->Integers;
+	}
+	int64 Group;
+	if (const FAttributeProto* MomentumAttribute = FModelProto::FindElementInArray(TEXT("Group"), InNodeProto.Attribute, /*bMustValueBeFound*/false))
+	{
+		Group = MomentumAttribute->I;
+	}
+	else
+	{
+		Group = 1;
+	}
+	TArray<int64> KernelShape;
+	if (const FAttributeProto* MomentumAttribute = FModelProto::FindElementInArray(TEXT("KernelShape"), InNodeProto.Attribute, /*bMustValueBeFound*/false))
+	{
+		KernelShape = MomentumAttribute->Integers;
+	}
+	TArray<int64> Pads;
+	if (const FAttributeProto* MomentumAttribute = FModelProto::FindElementInArray(TEXT("Pads"), InNodeProto.Attribute, /*bMustValueBeFound*/false))
+	{
+		Pads = MomentumAttribute->Integers;
+	}
+	TArray<int64> Strides;
+	if (const FAttributeProto* MomentumAttribute = FModelProto::FindElementInArray(TEXT("Strides"), InNodeProto.Attribute, /*bMustValueBeFound*/false))
+	{
+		Strides = MomentumAttribute->Integers;
+	}
+	return FConvOperator(AutoPad, Dilations, Group, KernelShape, Pads, Strides);
+}
+
+
+
+/* FConvOperator structors
+ *****************************************************************************/
+
+FConvOperator::FConvOperator(const FNodeProto& InNodeProto)
+	: FConvOperator(FPrivateConvOperator::NodeProtoToConvOperator(InNodeProto))
+{
+}
+
+FConvOperator::FConvOperator(const EAutoPad InAutoPad, const TArray<int64>& InDilations, const int64 InGroup, const TArray<int64>& InKernelShape, const TArray<int64>& InPads, const TArray<int64>& InStrides)
+	: FConvBaseOperator(TEXT("Conv"), 11, -1, InAutoPad, InDilations, InGroup, InKernelShape, InPads, InStrides, /*bIsTransposed*/false)
+{
+}
+
+FConvOperator::~FConvOperator()
+{
+}
+
+
+
+/* FConvOperator protected functions
+ *****************************************************************************/
+
+bool FConvOperator::SetAndConfigureStrides(const int32 InNumberConvolutionalDimensions)
+{
+	// Create and fill Strides (which might be empty)
+	if (Strides.Num() < 1)
+	{
+		Strides.Init(1, InNumberConvolutionalDimensions, /*bShouldCreate64BitVersion*/true);
+	}
+	return true;
+}
