NeuralNetworkInference (NNI) plugin:

diff --git a/Engine/Plugins/Experimental/NeuralNetworkInference/Source/NeuralNetworkInference/Private/NeuralOperators/ConvTransposeOperator.cpp b/Engine/Plugins/Experimental/NeuralNetworkInference/Source/NeuralNetworkInference/Private/NeuralOperators/ConvTransposeOperator.cpp
new file mode 100644
--- /dev/null
+++ b/Engine/Plugins/Experimental/NeuralNetworkInference/Source/NeuralNetworkInference/Private/NeuralOperators/ConvTransposeOperator.cpp
@@ -0,0 +1,457 @@
+// Copyright Epic Games, Inc. All Rights Reserved.
+
+#include "NeuralOperators/ConvTransposeOperator.h"
+#include "NeuralNetworkInferenceShaders/ConvTransposeCS.h"
+#include "NeuralNetworkInferenceUtils.h"
+#include "NeuralNetworkInferenceUtilsGPU.h"
+
+
+
+/* FPrivateConvTransposeOperator auxiliary functions
+ *****************************************************************************/
+
+class FPrivateConvTransposeOperator
+{
+public:
+	static FConvTransposeOperator NodeProtoToConvTransposeOperator(const FNodeProto& InNodeProto);
+
+	static bool ZerosFromConvStrides(FNeuralInt64ArrayUInt32Buffer& OutZeros, const FNeuralInt64ArrayUInt32Buffer& InStrides, const int32 InNumberConvolutionalDimensions);
+};
+
+FConvTransposeOperator FPrivateConvTransposeOperator::NodeProtoToConvTransposeOperator(const FNodeProto& InNodeProto)
+{
+	FConvBaseOperator::EAutoPad AutoPad = FConvBaseOperator::EAutoPad::NotSet;
+	if (const FAttributeProto* EpsilonAttribute = FModelProto::FindElementInArray(TEXT("AutoPad"), InNodeProto.Attribute, /*bMustValueBeFound*/false))
+	{
+		if (EpsilonAttribute->S == TEXT("NOTSET"))
+		{
+			AutoPad = FConvBaseOperator::EAutoPad::NotSet;
+		}
+		else if (EpsilonAttribute->S == TEXT("SAME_UPPER"))
+		{
+			AutoPad = FConvBaseOperator::EAutoPad::SameUpper;
+		}
+		else if (EpsilonAttribute->S == TEXT("SAME_LOWER"))
+		{
+			AutoPad = FConvBaseOperator::EAutoPad::SameLower;
+		}
+		else if (EpsilonAttribute->S == TEXT("VALID"))
+		{
+			AutoPad = FConvBaseOperator::EAutoPad::Valid;
+		}
+		else
+		{
+			UE_LOG(LogNeuralNetworkInference, Warning, TEXT("FConvBaseOperator(): Unknown EpsilonAttribute->S = %s."), *EpsilonAttribute->S);
+		}
+	}
+	TArray<int64> Dilations;
+	if (const FAttributeProto* MomentumAttribute = FModelProto::FindElementInArray(TEXT("Dilations"), InNodeProto.Attribute, /*bMustValueBeFound*/false))
+	{
+		Dilations = MomentumAttribute->Integers;
+	}
+	int64 Group;
+	if (const FAttributeProto* MomentumAttribute = FModelProto::FindElementInArray(TEXT("Group"), InNodeProto.Attribute, /*bMustValueBeFound*/false))
+	{
+		Group = MomentumAttribute->I;
+	}
+	else
+	{
+		Group = 1;
+	}
+	TArray<int64> KernelShape;
+	if (const FAttributeProto* MomentumAttribute = FModelProto::FindElementInArray(TEXT("KernelShape"), InNodeProto.Attribute, /*bMustValueBeFound*/false))
+	{
+		KernelShape = MomentumAttribute->Integers;
+	}
+	TArray<int64> Pads;
+	if (const FAttributeProto* MomentumAttribute = FModelProto::FindElementInArray(TEXT("Pads"), InNodeProto.Attribute, /*bMustValueBeFound*/false))
+	{
+		Pads = MomentumAttribute->Integers;
+	}
+	TArray<int64> Strides;
+	if (const FAttributeProto* MomentumAttribute = FModelProto::FindElementInArray(TEXT("Strides"), InNodeProto.Attribute, /*bMustValueBeFound*/false))
+	{
+		Strides = MomentumAttribute->Integers;
+	}
+	TArray<int64> Zeros;
+	if (const FAttributeProto* MomentumAttribute = FModelProto::FindElementInArray(TEXT("Zeros"), InNodeProto.Attribute, /*bMustValueBeFound*/false))
+	{
+		Zeros = MomentumAttribute->Integers;
+	}
+	TArray<int64> OutputPadding;
+	if (const FAttributeProto* MomentumAttribute = FModelProto::FindElementInArray(TEXT("OutputPadding"), InNodeProto.Attribute, /*bMustValueBeFound*/false))
+	{
+		OutputPadding = MomentumAttribute->Integers;
+	}
+	TArray<int64> OutputShape;
+	if (const FAttributeProto* MomentumAttribute = FModelProto::FindElementInArray(TEXT("OutputShape"), InNodeProto.Attribute, /*bMustValueBeFound*/false))
+	{
+		OutputShape = MomentumAttribute->Integers;
+	}
+	return FConvTransposeOperator(AutoPad, Dilations, Group, KernelShape, OutputPadding, OutputShape, Pads, Strides);
+}
+
+bool FPrivateConvTransposeOperator::ZerosFromConvStrides(FNeuralInt64ArrayUInt32Buffer& OutZeros, const FNeuralInt64ArrayUInt32Buffer& InStrides, const int32 InNumberConvolutionalDimensions)
+{
+	// OutZeros from InStrides
+	if (InStrides.Num() > 0)
+	{
+		TArray<int64> ZerosArrayTemporary;
+		ZerosArrayTemporary.SetNumUninitialized(InStrides.Num());
+		for (int32 DimensionIndex = 0; DimensionIndex < ZerosArrayTemporary.Num(); ++DimensionIndex)
+		{
+			ZerosArrayTemporary[DimensionIndex] = InStrides[DimensionIndex] - 1;
+		}
+		OutZeros.Move(ZerosArrayTemporary);
+	}
+	// OutZeros = {0, 0, ...} if no stride
+	else
+	{
+		OutZeros.Init(0, InNumberConvolutionalDimensions, /*bShouldCreate64BitVersion*/true);
+	}
+	return true;
+}
+
+
+
+/* FConvTransposeOperator structors
+ *****************************************************************************/
+
+FConvTransposeOperator::FConvTransposeOperator(const FNodeProto& InNodeProto)
+	: FConvTransposeOperator(FPrivateConvTransposeOperator::NodeProtoToConvTransposeOperator(InNodeProto))
+{
+}
+
+FConvTransposeOperator::FConvTransposeOperator(const FConvBaseOperator::EAutoPad InAutoPad, const TArray<int64>& InDilations, const int64 InGroup, const TArray<int64>& InKernelShape,
+	const TArray<int64>& InOutputPadding, const TArray<int64>& InOutputShape, const TArray<int64>& InPads, const TArray<int64>& InStrides)
+	: FConvBaseOperator(TEXT("ConvTranpose"), 11, -1, InAutoPad, InDilations, InGroup, InKernelShape, InPads, InStrides, /*bIsTransposed*/true, InOutputPadding, InOutputShape)
+	, bIsXWithZerosNeeded(false)
+	, bIsWeightTensorFlipped(false)
+{
+	// Set bIsXWithZerosNeeded to true if an auxiliary tensor is needed, false otherwise
+	for (const int64 Stride : Strides)
+	{
+		if (Stride > 1)
+		{
+			bIsXWithZerosNeeded = true;
+			break;
+		}
+	}
+}
+
+FConvTransposeOperator::~FConvTransposeOperator()
+{
+}
+
+
+
+/* FConvTransposeOperator public functions
+ *****************************************************************************/
+
+int32 FConvTransposeOperator::GetNumberAuxiliaryTensors() const
+{
+	return (bIsXWithZerosNeeded ? 1 : 0);
+}
+
+bool FConvTransposeOperator::ConfigureOutputAndInternalVariablesAndSanityChecks()
+{
+	// FConvBaseOperator sanity checks, it will directly modify bIsLoaded
+	if (!FConvBaseOperator::ConfigureOutputAndInternalVariablesAndSanityChecks())
+	{
+		return bIsLoaded;
+	}
+	bIsLoaded = false;
+	// bIsWeightTensorFlipped sanity check
+	if (!bIsWeightTensorFlipped)
+	{
+		UE_LOG(LogNeuralNetworkInference, Warning,
+			TEXT("FConvTransposeOperator::ConfigureOutputAndInternalVariablesAndSanityChecks(): The weight tensor of this operator must be manually flipped with FNeuralTensor::Flip() before using this operator and SetWhetherWeightTensorWasFlipped(true) called. Dynamic weights not supported yet."));
+		return bIsLoaded;
+	}
+	// Zeros sanity check
+	const TArray<FNeuralTensor*>& InputTensors = GetInputTensorsConst();
+	const FNeuralTensor& X = *InputTensors[0];
+	const int32 NumberConvolutionalDimensions = X.GetNumberDimensions() - 2;
+	if (Zeros.Num() != NumberConvolutionalDimensions)
+	{
+		UE_LOG(LogNeuralNetworkInference, Warning, TEXT("FConvTransposeOperator::ConfigureOutputAndInternalVariablesAndSanityChecks(): Zeros.Num() == NumberConvolutionalDimensions failed, %d != %d."), Zeros.Num(), NumberConvolutionalDimensions);
+		return bIsLoaded;
+	}
+	bIsLoaded = true;
+	return bIsLoaded;
+}
+
+void FConvTransposeOperator::ToGPU_RenderThread()
+{
+	// Zeros and XSizesForGPU to GPU
+	if (bIsXWithZerosNeeded)
+	{
+		if (Zeros.IsReadBufferValid() || XSizesForGPU.IsReadBufferValid())
+		{
+			UE_LOG(LogNeuralNetworkInference, Warning, TEXT("FConvBaseOperator::ToGPU_RenderThread():"
+				" Zeros, XSizesForGPU and/or XSizesForGPU FReadBuffers were not nullptr."));
+		}
+		Zeros.CreateAndLoadSRVBuffer(TEXT("ConvTranpose_Zeros"));
+		XSizesForGPU.CreateAndLoadSRVBuffer(TEXT("ConvTranpose_XSizesForGPU"));
+	}
+	FConvBaseOperator::ToGPU_RenderThread();
+}
+
+#define MULTICHANNEL_FOR_LOOPS_INIT() \
+	for (int64 BatchIndex = 0; BatchIndex < XSizes[0]; ++BatchIndex) \
+	{ \
+		const int64 BatchOffset = BatchIndex * XSizes[1]; \
+		for (int64 ChannelIndex = 0; ChannelIndex < XSizes[1]; ++ChannelIndex) \
+		{ \
+			const int64 BatchChannelOffset = BatchOffset + ChannelIndex; \
+			const int64 XBatchChannelOffset = BatchChannelOffset * XSizes[2]; \
+			const int64 XWithOffsetBatchChannelOffset = BatchChannelOffset * XWithZerosSizes[2];
+
+#define MULTICHANNEL_FOR_LOOPS_END() \
+		} \
+	}
+
+void FConvTransposeOperator::ForwardCPU()
+{
+	// X --> XWithZeros (i.e., AuxiliaryTensors[0])
+	if (bIsXWithZerosNeeded)
+	{
+		FNeuralTensor& XWithZeros = *AuxiliaryTensors[0];
+		const FNeuralTensor& X = *GetInputTensorsConst()[0];
+		const int32 NumberDimensions = X.GetNumberDimensions();
+		const TArray<int64>& XSizes = X.GetSizes();
+		const TArray<int64>& XWithZerosSizes = XWithZeros.GetSizes();
+		// 1D
+		if (NumberDimensions == 3)
+		{
+			MULTICHANNEL_FOR_LOOPS_INIT()
+				for (int64 D1Index = 0; D1Index < XSizes[2]; ++D1Index)
+				{
+					const int64 D1IndexWithZeros = D1Index * (Zeros[0] + 1);
+					XWithZeros.At<float>(XWithOffsetBatchChannelOffset + D1IndexWithZeros) = X.At<float>(XBatchChannelOffset + D1Index);
+				}
+			MULTICHANNEL_FOR_LOOPS_END()
+		}
+		// 2D
+		else if (NumberDimensions == 4)
+		{
+			MULTICHANNEL_FOR_LOOPS_INIT()
+				for (int64 D1Index = 0; D1Index < XSizes[2]; ++D1Index)
+				{
+					const int64 D1IndexWithZeros = D1Index * (Zeros[0] + 1);
+					for (int64 D2Index = 0; D2Index < XSizes[3]; ++D2Index)
+					{
+						const int64 D2IndexWithZeros = D2Index * (Zeros[1] + 1);
+						XWithZeros.At<float>((XWithOffsetBatchChannelOffset + D1IndexWithZeros) * XWithZerosSizes[3] + D2IndexWithZeros) = X.At<float>((XBatchChannelOffset + D1Index) * XSizes[3] + D2Index);
+					}
+				}
+			MULTICHANNEL_FOR_LOOPS_END()
+		}
+		// 3D
+		else if (NumberDimensions == 5)
+		{
+			MULTICHANNEL_FOR_LOOPS_INIT()
+				for (int64 D1Index = 0; D1Index < XSizes[2]; ++D1Index)
+				{
+					const int64 D1IndexWithZeros = D1Index * (Zeros[0] + 1);
+					for (int64 D2Index = 0; D2Index < XSizes[3]; ++D2Index)
+					{
+						const int64 D2IndexWithZeros = D2Index * (Zeros[1] + 1);
+						for (int64 D3Index = 0; D3Index < XSizes[4]; ++D3Index)
+						{
+							const int64 D3IndexWithZeros = D3Index * (Zeros[2] + 1);
+							XWithZeros.At<float>(((XWithOffsetBatchChannelOffset + D1IndexWithZeros) * XWithZerosSizes[3] + D2IndexWithZeros) * XWithZerosSizes[4] + D3IndexWithZeros)
+								= X.At<float>(((XBatchChannelOffset + D1Index) * XSizes[3] + D2Index) * XSizes[4] + D3Index);
+						}
+					}
+				}
+			MULTICHANNEL_FOR_LOOPS_END()
+		}
+		// 5D
+		else if (NumberDimensions == 7)
+		{
+			MULTICHANNEL_FOR_LOOPS_INIT()
+				for (int64 D1Index = 0; D1Index < XSizes[2]; ++D1Index)
+				{
+					const int64 D1IndexWithZeros = D1Index * (Zeros[0] + 1);
+					for (int64 D2Index = 0; D2Index < XSizes[3]; ++D2Index)
+					{
+						const int64 D2IndexWithZeros = D2Index * (Zeros[1] + 1);
+						for (int64 D3Index = 0; D3Index < XSizes[4]; ++D3Index)
+						{
+							const int64 D3IndexWithZeros = D3Index * (Zeros[2] + 1);
+							for (int64 D4Index = 0; D4Index < XSizes[5]; ++D4Index)
+							{
+								const int64 D4IndexWithZeros = D4Index * (Zeros[3] + 1);
+								for (int64 D5Index = 0; D5Index < XSizes[6]; ++D5Index)
+								{
+									const int64 D5IndexWithZeros = D5Index * (Zeros[4] + 1);
+									XWithZeros.At<float>(((((XWithOffsetBatchChannelOffset + D1IndexWithZeros) * XWithZerosSizes[3] + D2IndexWithZeros) * XWithZerosSizes[4] + D3IndexWithZeros) * XWithZerosSizes[5] + D4IndexWithZeros) * XWithZerosSizes[6] + D5IndexWithZeros)
+										= X.At<float>(((((XBatchChannelOffset + D1Index) * XSizes[3] + D2Index) * XSizes[4] + D3Index) * XSizes[5] + D4Index) * XSizes[6] + D5Index);
+								}
+							}
+						}
+					}
+				}
+			MULTICHANNEL_FOR_LOOPS_END()
+		}
+		// nD
+		else
+		{
+			UE_LOG(LogNeuralNetworkInference, Warning, TEXT("%dD transposed convolution not implemented yet."), NumberDimensions);
+			return;
+		}
+	}
+	FConvBaseOperator::ForwardCPU();
+}
+
+void FConvTransposeOperator::ForwardGPU_RenderThread(FRDGBuilder* InOutGraphBuilder)
+{
+	// X --> XWithZeros (i.e., AuxiliaryTensors[0])
+	if (bIsXWithZerosNeeded)
+	{
+		// Sanity checks
+		if (!FNeuralNetworkInferenceUtilsGPU::GPUSanityChecks(InOutGraphBuilder, bIsLoaded))
+		{
+			UE_LOG(LogNeuralNetworkInference, Warning, TEXT("FConvBaseOperator::ForwardGPU_RenderThread(): Sanity checks failed."));
+			return;
+		}
+		const int32 NumberConvolutionalDimensions = AuxiliaryTensors[0]->GetNumberDimensions() - 2;
+		if (NumberConvolutionalDimensions != 1 && NumberConvolutionalDimensions != 2 && NumberConvolutionalDimensions != 3 && NumberConvolutionalDimensions != 5)
+		{
+			UE_LOG(LogNeuralNetworkInference, Warning, TEXT("FConvBaseOperator::ForwardGPU_RenderThread(): Only 1D, 2D, 3D, or 5D convolutions implemented, but not %dD."), NumberConvolutionalDimensions);
+			return;
+		}
+		// Set parameters
+		FConvTransposeCS::FParameters* Parameters = InOutGraphBuilder->AllocParameters<FConvTransposeCS::FParameters>();
+		// Input SRV variables - Fill ReadBuffer's the first time
+		const FNeuralTensor& X = *GetInputTensorsConst()[0];
+		// Input SRV variables
+		Parameters->Zeros = Zeros.GetUInt32SRV();
+		Parameters->XSizes = XSizesForGPU.GetUInt32SRV();
+		Parameters->XWithZerosSizes = XOrXWithZerosSizesForGPU.GetUInt32SRV();
+		// Input variables
+		Parameters->XVolume = X.Num();
+		// nD convolution
+		const EConvMode ConvMode = (NumberConvolutionalDimensions == 1 ? EConvMode::D1
+			: NumberConvolutionalDimensions == 2 ? EConvMode::D2
+			: NumberConvolutionalDimensions == 3 ? EConvMode::D3
+			: /*NumberConvolutionalDimensions > 3 ?*/ EConvMode::DN);
+		if (ConvMode == EConvMode::DN)
+		{
+			Parameters->NumberConvolutionalDimensions = NumberConvolutionalDimensions;
+		}
+		// SRV/UAV variables
+		Parameters->XSRV = X.GetBufferSRVRef();
+		Parameters->XWithZerosUAV = AuxiliaryTensors[0]->GetBufferUAVRef();
+		// Set shader
+		FConvTransposeCS::FPermutationDomain PermutationVector;
+		PermutationVector.Set<FConvTransposeCS::FConvMode>(ConvMode);
+		TShaderMapRef<FConvTransposeCS> ComputeShader(GetGlobalShaderMap(GMaxRHIFeatureLevel), PermutationVector);
+		// Run shader
+		const uint32 ThreadGroupCountValueX = FMath::DivideAndRoundUp(Parameters->XVolume, FConvTransposeCS::THREADGROUP_SIZE_X);
+		FComputeShaderUtils::AddPass(
+			*InOutGraphBuilder,
+			RDG_EVENT_NAME("FConvTransposeCS - %dD convolution, Operator: %s", NumberConvolutionalDimensions, *Name),
+			ComputeShader,
+			Parameters,
+			FIntVector(ThreadGroupCountValueX, 1, 1));
+	}
+	FConvBaseOperator::ForwardGPU_RenderThread(InOutGraphBuilder);
+}
+
+
+
+/* FConvTransposeOperator protected functions
+ *****************************************************************************/
+
+bool FConvTransposeOperator::SetAndConfigureStrides(const int32 InNumberConvolutionalDimensions)
+{
+	// Create and fill Zeros from Strides (if not empty)
+	if (!FPrivateConvTransposeOperator::ZerosFromConvStrides(Zeros, Strides, InNumberConvolutionalDimensions))
+	{
+		UE_LOG(LogNeuralNetworkInference, Warning, TEXT("FConvTransposeOperator::SetAndConfigureStrides(): ZerosFromConvStrides() failed."));
+		return false;
+	}
+	// Zeros sanity checks
+	if (Zeros.Num() > 0 && Zeros.Num() != InNumberConvolutionalDimensions) // Right dimensions
+	{
+		UE_LOG(LogNeuralNetworkInference, Warning, TEXT("FConvTransposeOperator::SetAndConfigureStrides(): Zeros.Num() == 0 || Zeros.Num() == InNumberConvolutionalDimensions failed (%d != 0 || %d)."),
+			Zeros.Num(), InNumberConvolutionalDimensions);
+		return false;
+	}
+	// Sanity checks
+	for (const int64 Zero : Zeros) // Zeros >= 0 on each dimension
+	{
+		if (Zero < 0)
+		{
+			UE_LOG(LogNeuralNetworkInference, Warning, TEXT("FConvTransposeOperator::SetAndConfigureStrides(): Zeros (%d) < 1 found! Zeros must be >= 1."), Zero);
+			return false;
+		}
+		else if (!bIsXWithZerosNeeded && Zero > 0)
+		{
+			UE_LOG(LogNeuralNetworkInference, Warning, TEXT("FConvTransposeOperator::SetAndConfigureStrides(): bIsXWithZerosNeeded was false but Zeros (%d) > 0 found! Only one of the 2 can be true."), Zero);
+			return false;
+		}
+	}
+	// Set ConvolutionStridesIfTransposedConvolution
+	if (Strides.Num() > 0)
+	{
+		ConvolutionStridesIfTransposedConvolution = Strides.GetInt64();
+	}
+	else
+	{
+		ConvolutionStridesIfTransposedConvolution.Init(1, InNumberConvolutionalDimensions);
+	}
+	// Strides = {1, 1, ...}
+	Strides.Init(1, InNumberConvolutionalDimensions, /*bShouldCreate64BitVersion*/true);
+	return true;
+}
+
+bool FConvTransposeOperator::SetAndConfigureAuxiliaryTensor()
+{
+	// Sanity check
+	if (Zeros.Num() < 1)
+	{
+		UE_LOG(LogNeuralNetworkInference, Warning, TEXT("FConvTransposeOperator::SetAndConfigureAuxiliaryTensor(): Zeros.Num() should not be 0."));
+		return false;
+	}
+	if (bIsXWithZerosNeeded)
+	{
+		const FNeuralTensor& X = *GetInputTensorsConst()[0];
+		const TArray<int64>& XSizes = X.GetSizes();
+		// Sanity checks
+		if (XSizes.Num() < 2)
+		{
+			UE_LOG(LogNeuralNetworkInference, Warning, TEXT("FConvTransposeOperator::SetAndConfigureAuxiliaryTensor(): XSizes.Num() should be at least 2, not %d."), XSizes.Num());
+			return false;
+		}
+		else if (AuxiliaryTensors.Num() != 1)
+		{
+			UE_LOG(LogNeuralNetworkInference, Warning, TEXT("FConvTransposeOperator::SetAndConfigureAuxiliaryTensor(): AuxiliaryTensors.Num() should be 1, not %d."), AuxiliaryTensors.Num());
+			return false;
+		}
+		// Estimate XWithZerosSizes
+		TArray<int64> XWithZerosSizes = XSizes; // XWithZerosSizes[0] and [1] are the same than XSizes
+		for (int32 DimensionIndex = 2; DimensionIndex < XWithZerosSizes.Num(); ++DimensionIndex)
+		{
+			if (XSizes[DimensionIndex] != 1) // Output dimension size = D --> XWithZerosSize = D + Zx(D-1)
+			{
+				XWithZerosSizes[DimensionIndex] = XSizes[DimensionIndex] + Zeros[DimensionIndex - 2] * (XSizes[DimensionIndex] - 1);
+			}
+			else // Output dimension size = D --> XWithZerosSize = D + Z
+			{
+				XWithZerosSizes[DimensionIndex] = XSizes[DimensionIndex] + Zeros[DimensionIndex - 2];
+			}
+		}
+		// Set XWithZeros (i.e., AuxiliaryTensors[0])
+		FNeuralTensor& XWithZeros = *AuxiliaryTensors[0];
+		// Initialize XWithZeros
+		XWithZeros.SetNumUninitialized(XWithZerosSizes, X.GetDataType());
+		XWithZeros.SetTo<float>(0.f);
+		XWithZeros.SetTensorTypeGPU(ENeuralTensorTypeGPU::IntermediateInitialized);
+		// Set XSizesForGPU
+		XSizesForGPU.Copy(X.GetSizes());
+	}
+	return true;
+}
diff --git a/Engine/Plugins/Experimental/NeuralNetworkInference/Source/NeuralNetworkInference/Private/NeuralOperators/CosOperator.cpp b/Engine/Plugins/Experimental/NeuralNetworkInference/Source/NeuralNetworkInference/Private/NeuralOperators/CosOperator.cpp
new file mode 100644
--- /dev/null
+++ b/Engine/Plugins/Experimental/NeuralNetworkInference/Source/NeuralNetworkInference/Private/NeuralOperators/CosOperator.cpp
@@ -0,0 +1,17 @@
+// Copyright Epic Games, Inc. All Rights Reserved.
+
+#include "NeuralOperators/CosOperator.h"
+
+
+
+/* FCosOperator structors
+ *****************************************************************************/
+
+FCosOperator::FCosOperator(const bool bIsInlinedTensor)
+	: IElementWiseOperator(TEXT("Cos"), 7, EElementWiseOperator::Cos, bIsInlinedTensor)
+{
+}
+
+FCosOperator::~FCosOperator()
+{
+}
diff --git a/Engine/Plugins/Experimental/NeuralNetworkInference/Source/NeuralNetworkInference/Private/NeuralOperators/CoshOperator.cpp b/Engine/Plugins/Experimental/NeuralNetworkInference/Source/NeuralNetworkInference/Private/NeuralOperators/CoshOperator.cpp
new file mode 100644
--- /dev/null
+++ b/Engine/Plugins/Experimental/NeuralNetworkInference/Source/NeuralNetworkInference/Private/NeuralOperators/CoshOperator.cpp
@@ -0,0 +1,28 @@
+// Copyright Epic Games, Inc. All Rights Reserved.
+
+#include "NeuralOperators/CoshOperator.h"
+#include <cmath>
+
+
+
+/* FCoshOperator structors
+ *****************************************************************************/
+
+FCoshOperator::FCoshOperator(const bool bIsInlinedTensor)
+	: IElementWiseOperator(TEXT("Cosh"), 9, EElementWiseOperator::Cosh, bIsInlinedTensor)
+{
+}
+
+FCoshOperator::~FCoshOperator()
+{
+}
+
+
+
+/* FCoshOperator public functions
+ *****************************************************************************/
+
+void FCoshOperator::ForwardCPU()
+{
+	ForwardCPUWithFunction([](const float InValue) { return std::cosh(InValue); });
+}
diff --git a/Engine/Plugins/Experimental/NeuralNetworkInference/Source/NeuralNetworkInference/Private/NeuralOperators/DivOperator.cpp b/Engine/Plugins/Experimental/NeuralNetworkInference/Source/NeuralNetworkInference/Private/NeuralOperators/DivOperator.cpp
new file mode 100644
--- /dev/null
+++ b/Engine/Plugins/Experimental/NeuralNetworkInference/Source/NeuralNetworkInference/Private/NeuralOperators/DivOperator.cpp
@@ -0,0 +1,17 @@
+// Copyright Epic Games, Inc. All Rights Reserved.
+
+#include "NeuralOperators/DivOperator.h"
+
+
+
+/* FDivOperator structors
+ *****************************************************************************/
+
+FDivOperator::FDivOperator(const TSet<uint32>& InPotentialInlinedTensors)
+	: IMultidirectionalBroadcastOperator(TEXT("Div"), 13, EMultidirectionalBroadcastOperator::Div, InPotentialInlinedTensors)
+{
+}
+
+FDivOperator::~FDivOperator()
+{
+}
diff --git a/Engine/Plugins/Experimental/NeuralNetworkInference/Source/NeuralNetworkInference/Private/NeuralOperators/ElementWiseOperator.cpp b/Engine/Plugins/Experimental/NeuralNetworkInference/Source/NeuralNetworkInference/Private/NeuralOperators/ElementWiseOperator.cpp
new file mode 100644
--- /dev/null
+++ b/Engine/Plugins/Experimental/NeuralNetworkInference/Source/NeuralNetworkInference/Private/NeuralOperators/ElementWiseOperator.cpp
@@ -0,0 +1,173 @@
+// Copyright Epic Games, Inc. All Rights Reserved.
+
+#include "NeuralOperators/ElementWiseOperator.h"
+#include "NeuralNetworkInferenceShaders/ElementWiseCS.h"
+#include "NeuralNetworkInferenceUtils.h"
+#include "NeuralNetworkInferenceUtilsGPU.h"
+
+
+
+/* IElementWiseOperator structors
+ *****************************************************************************/
+
+IElementWiseOperator::IElementWiseOperator(const FString& InName, const int32 InVersion, const EElementWiseOperator InElementWiseOperator,
+	const bool bIsInlinedTensor, const TArray<float>& InAttributes)
+	: FNeuralOperator(InName, InVersion, (bIsInlinedTensor ? 0 : -1))
+	, Attributes(InAttributes)
+	, ElementWiseOperator(InElementWiseOperator)
+{
+}
+
+IElementWiseOperator::~IElementWiseOperator()
+{
+}
+
+
+
+/* IElementWiseOperator public functions
+ *****************************************************************************/
+
+bool IElementWiseOperator::ConfigureOutputAndInternalVariablesAndSanityChecks()
+{
+	bIsLoaded = false;
+	// Attributes
+	if (Attributes.Num() > 1)
+	{
+		UE_LOG(LogNeuralNetworkInference, Warning, TEXT("IElementWiseOperator::ConfigureOutputAndInternalVariablesAndSanityChecks(): So far only IElementWiseOperators with 0 or 1"
+			" attributes, not implemented for %d."), Attributes.Num());
+		return bIsLoaded;
+	}
+	// Input sanity checks
+	const TArray<FNeuralTensor*>& InputTensors = GetInputTensorsConst();
+	if (!FNeuralNetworkInferenceUtils::SizeSanityChecks(InputTensors, 1))
+	{
+		return bIsLoaded;
+	}
+	const FNeuralTensor& InputTensor = *InputTensors[0];
+	// Output (if not inlined)
+	if (InlinedTensor < 0)
+	{
+		TArray<FNeuralTensor*>& OutputTensors = GetOutputTensorsNoConst();
+		// Sanity checks
+		if (!FNeuralNetworkInferenceUtils::SizeSanityChecks(OutputTensors, 1))
+		{
+			return bIsLoaded;
+		}
+		// Set and configure output tensor
+		FNeuralTensor& OutputTensor = *OutputTensors.Last();
+		OutputTensor.SetNumUninitialized(InputTensor);
+	}
+	// Return true
+	bIsLoaded = true;
+	return bIsLoaded;
+}
+
+void IElementWiseOperator::ForwardGPU_RenderThread(FRDGBuilder* InOutGraphBuilder)
+{
+	// Sanity checks
+	if (!FNeuralNetworkInferenceUtilsGPU::GPUSanityChecks(InOutGraphBuilder, bIsLoaded))
+	{
+		UE_LOG(LogNeuralNetworkInference, Warning, TEXT("IElementWiseOperator::ForwardGPU_RenderThread(): Sanity checks failed."));
+		return;
+	}
+	// Set parameters
+	FElementWiseCS::FParameters* Parameters = InOutGraphBuilder->AllocParameters<FElementWiseCS::FParameters>();
+	// Input variables
+	const FNeuralTensor& InputTensor = GetInputTensorConst();
+	Parameters->TensorSize = (uint32)GetInputTensorConst().Num();
+	if (Attributes.Num() > 0)
+	{
+		Parameters->Attribute = Attributes[0];
+	}
+	// SRV/UAV variables
+	Parameters->InputSRV = InputTensor.GetBufferSRVRef();
+	Parameters->OutputUAV = (InlinedTensor < 0 ? GetOutputTensorNoConst() : GetInputTensorNoConst()).GetBufferUAVRef(); // Not inlined vs. inlined
+	// Set shader
+	FElementWiseCS::FPermutationDomain PermutationVector;
+	PermutationVector.Set<FElementWiseCS::FShaderType>(ElementWiseOperator);
+	PermutationVector.Set<FElementWiseCS::FIsInlined>((InlinedTensor > -1));
+	TShaderMapRef<FElementWiseCS> ComputeShader(GetGlobalShaderMap(GMaxRHIFeatureLevel), PermutationVector);
+	// Run shader
+	const uint32 ThreadGroupCountValueX = FMath::DivideAndRoundUp((uint32)InputTensor.Num(), FElementWiseCS::THREADGROUP_SIZE_X);
+	FComputeShaderUtils::AddPass(
+		*InOutGraphBuilder,
+		RDG_EVENT_NAME("FElementWiseCS() - Operator: %s, InlinedTensor: %d", *Name, InlinedTensor),
+		ComputeShader,
+		Parameters,
+		FIntVector(ThreadGroupCountValueX, 1, 1));
+}
+
+
+
+/* IElementWiseOperator protected functions
+ *****************************************************************************/
+
+void IElementWiseOperator::ForwardCPUWithFunction(float InOperatorFunction(const float))
+{
+	// Sanity checks
+	if (!bIsLoaded)
+	{
+		UE_LOG(LogNeuralNetworkInference, Warning, TEXT("IElementWiseOperator::ForwardCPUWithFunction(): bIsLoaded was false."));
+		return;
+	}
+	if (Attributes.Num() != 0)
+	{
+		UE_LOG(LogNeuralNetworkInference, Warning, TEXT("IElementWiseOperator::ForwardCPUWithFunction(): Attributes.Num() = %d != 0."), Attributes.Num());
+		return;
+	}
+	// Not inlined layer (output is different than input)
+	if (InlinedTensor < 0)
+	{
+		const FNeuralTensor& InputTensor = GetInputTensorConst();
+		FNeuralTensor& OutputTensor = GetOutputTensorNoConst();
+		for (int64 Index = 0; Index < InputTensor.Num(); ++Index)
+		{
+			OutputTensor.At<float>(Index) = InOperatorFunction(InputTensor.At<float>(Index));
+		}
+	}
+	// Inlined layer
+	else
+	{
+		FNeuralTensor& InputTensor = GetInputTensorNoConst();
+		for (int64 Index = 0; Index < InputTensor.Num(); ++Index)
+		{
+			float& TensorValue = InputTensor.At<float>(Index);
+			TensorValue = InOperatorFunction(TensorValue);
+		}
+	}
+}
+
+void IElementWiseOperator::ForwardCPUWithFunction(float InOperatorFunction(const float, const float))
+{
+	// Sanity checks
+	if (!bIsLoaded)
+	{
+		UE_LOG(LogNeuralNetworkInference, Warning, TEXT("IElementWiseOperator::ForwardCPUWithFunction(): bIsLoaded was false."));
+		return;
+	}
+	if (Attributes.Num() != 1)
+	{
+		UE_LOG(LogNeuralNetworkInference, Warning, TEXT("IElementWiseOperator::ForwardCPUWithFunction(): Attributes.Num() = %d != 1."), Attributes.Num());
+		return;
+	}
+	// Not inlined layer (output is different than input)
+	if (InlinedTensor < 0)
+	{
+		const FNeuralTensor& InputTensor = GetInputTensorConst();
+		FNeuralTensor& OutputTensor = GetOutputTensorNoConst();
+		for (int64 Index = 0; Index < InputTensor.Num(); ++Index)
+		{
+			OutputTensor.At<float>(Index) = InOperatorFunction(InputTensor.At<float>(Index), Attributes[0]);
+		}
+	}
+	// Inlined layer
+	else
+	{
+		FNeuralTensor& InputTensor = GetInputTensorNoConst();
+		for (int64 Index = 0; Index < InputTensor.Num(); ++Index)
+		{
+			float& TensorValue = InputTensor.At<float>(Index);
+			TensorValue = InOperatorFunction(TensorValue, Attributes[0]);
+		}
+	}
+}
diff --git a/Engine/Plugins/Experimental/NeuralNetworkInference/Source/NeuralNetworkInference/Private/NeuralOperators/ExpOperator.cpp b/Engine/Plugins/Experimental/NeuralNetworkInference/Source/NeuralNetworkInference/Private/NeuralOperators/ExpOperator.cpp
new file mode 100644
--- /dev/null
+++ b/Engine/Plugins/Experimental/NeuralNetworkInference/Source/NeuralNetworkInference/Private/NeuralOperators/ExpOperator.cpp
@@ -0,0 +1,17 @@
+// Copyright Epic Games, Inc. All Rights Reserved.
+
+#include "NeuralOperators/ExpOperator.h"
+
+
+
+/* FExpOperator structors
+ *****************************************************************************/
+
+FExpOperator::FExpOperator(const bool bIsInlinedTensor)
+	: IElementWiseOperator(TEXT("Exp"), 13, EElementWiseOperator::Exp, bIsInlinedTensor)
+{
+}
+
+FExpOperator::~FExpOperator()
+{
+}
diff --git a/Engine/Plugins/Experimental/NeuralNetworkInference/Source/NeuralNetworkInference/Private/NeuralOperators/FloorOperator.cpp b/Engine/Plugins/Experimental/NeuralNetworkInference/Source/NeuralNetworkInference/Private/NeuralOperators/FloorOperator.cpp
new file mode 100644
--- /dev/null
+++ b/Engine/Plugins/Experimental/NeuralNetworkInference/Source/NeuralNetworkInference/Private/NeuralOperators/FloorOperator.cpp
@@ -0,0 +1,17 @@
+// Copyright Epic Games, Inc. All Rights Reserved.
+
+#include "NeuralOperators/FloorOperator.h"
+
+
+
+/* FFloorOperator structors
+ *****************************************************************************/
+
+FFloorOperator::FFloorOperator(const bool bIsInlinedTensor)
+	: IElementWiseOperator(TEXT("Floor"), 13, EElementWiseOperator::Floor, bIsInlinedTensor)
+{
+}
+
+FFloorOperator::~FFloorOperator()
+{
+}
diff --git a/Engine/Plugins/Experimental/NeuralNetworkInference/Source/NeuralNetworkInference/Private/NeuralOperators/GemmOperator.cpp b/Engine/Plugins/Experimental/NeuralNetworkInference/Source/NeuralNetworkInference/Private/NeuralOperators/GemmOperator.cpp
new file mode 100644
--- /dev/null
+++ b/Engine/Plugins/Experimental/NeuralNetworkInference/Source/NeuralNetworkInference/Private/NeuralOperators/GemmOperator.cpp
@@ -0,0 +1,288 @@
+// Copyright Epic Games, Inc. All Rights Reserved.
+
+#include "NeuralOperators/GemmOperator.h"
+#include "NeuralNetworkInferenceShaders/GemmCS.h"
+#include "NeuralNetworkInferenceUtils.h"
+#include "NeuralNetworkInferenceUtilsGPU.h"
+
+
+
+/* FGemmOperator structors
+ *****************************************************************************/
+
+FGemmOperator::FGemmOperator(const FNodeProto& InNodeProto)
+	: FGemmOperator()
+{
+	if (const FAttributeProto* AlphaAttribute = FModelProto::FindElementInArray(TEXT("Alpha"), InNodeProto.Attribute, /*bMustValueBeFound*/false))
+	{
+		Alpha = AlphaAttribute->F;
+	}
+	if (const FAttributeProto* BetaAttribute = FModelProto::FindElementInArray(TEXT("Beta"), InNodeProto.Attribute, /*bMustValueBeFound*/false))
+	{
+		Beta = BetaAttribute->F;
+	}
+	if (const FAttributeProto* TransAAttribute = FModelProto::FindElementInArray(TEXT("TransA"), InNodeProto.Attribute, /*bMustValueBeFound*/false))
+	{
+		bTransA = (TransAAttribute->I != 0);
+	}
+	if (const FAttributeProto* TransBAttribute = FModelProto::FindElementInArray(TEXT("TransB"), InNodeProto.Attribute, /*bMustValueBeFound*/false))
+	{
+		bTransB = (TransBAttribute->I != 0);
+	}
+}
+
+FGemmOperator::FGemmOperator(const float InAlpha, const float InBeta, const bool bInTransA, const bool bInTransB)
+	: FNeuralOperator(TEXT("GEMM"), 13, -1)
+	, Alpha(InAlpha)
+	, Beta(InBeta)
+	, bTransA(bInTransA)
+	, bTransB(bInTransB)
+{
+}
+
+FGemmOperator::~FGemmOperator()
+{
+}
+
+
+
+/* FGemmOperator public functions
+ *****************************************************************************/
+
+bool FGemmOperator::ConfigureOutputAndInternalVariablesAndSanityChecks()
+{
+	bIsLoaded = false;
+	// Input sanity checks
+	const TArray<FNeuralTensor*>& InputTensors = GetInputTensorsConst();
+	if (!FNeuralNetworkInferenceUtils::SizeSanityChecks(InputTensors, 2, 3, 1, 2))
+	{
+		return bIsLoaded;
+	}
+	const int64 CommonK1 = InputTensors[0]->GetSize(bTransA ? 0 : 1);
+	const int64 CommonK2 = InputTensors[1]->GetSize(bTransB ? 1 : 0);
+	if (CommonK1 != CommonK2)
+	{
+		UE_LOG(LogNeuralNetworkInference, Warning, TEXT("FGemmOperator::ConfigureOutputAndInternalVariablesAndSanityChecks(): CommonK1 == CommonK2 failed: %d != %d"), CommonK1, CommonK2);
+		return bIsLoaded;
+	}
+	if (InlinedTensor >= 0)
+	{
+		UE_LOG(LogNeuralNetworkInference, Warning, TEXT("FGemmOperator::ConfigureOutputAndInternalVariablesAndSanityChecks(): FGemmOperator cannot be inlined (InlinedTensor = %d)."), InlinedTensor);
+		return bIsLoaded;
+	}
+	// Initialize output tensor
+	const int64 M = InputTensors[0]->GetSize(bTransA ? 1 : 0);
+	const int64 N = InputTensors[1]->GetSize(bTransB ? 0 : 1);
+	FNeuralTensor& OutputTensor = GetOutputTensorNoConst();
+	OutputTensor.SetNumUninitialized(TArray<int64>({ M, N }), InputTensors[0]->GetDataType());
+	// Output sanity checks
+	if (!FNeuralNetworkInferenceUtils::SizeSanityChecks(GetOutputTensorsNoConst(), 1, 1, 2, 3))
+	{
+		return bIsLoaded;
+	}
+	// Check bias tensor (C)
+	if (InputTensors.Num() == 3)
+	{
+		const FNeuralTensor& C = *InputTensors[2];
+		// If not same volume than OutputTensor (and not an scalar), then C is unidirectional broadcastable to output tensor. I.e., one of the following is true:
+		// Source: [REDACTED_PROFILE_URL]/onnx/blob/master/docs/Broadcasting.md#unidirectional-broadcasting
+		// - Tensor A and B both have exactly the same shape.
+		// - Tensor A and B all have the same number of dimensions and the length of each dimensions is either a common length or B's length is 1.
+		// - Tensor B has too few dimensions, and B can have its shapes prepended with a dimension of length 1 to satisfy property 2.
+		TArray<int64> CSize2D({ 1u, 1u });
+		if (OutputTensor.GetNumberDimensions() < C.GetNumberDimensions())
+		{
+			UE_LOG(LogNeuralNetworkInference, Warning, TEXT("FGemmOperator::ConfigureOutputAndInternalVariablesAndSanityChecks(): Dimension mismatch between the output tensor and"
+				" the bias tensor (C): %d vs. %d."), OutputTensor.GetNumberDimensions(), C.GetNumberDimensions());
+			return bIsLoaded;
+		}
+		for (int32 CDimension = 0; CDimension < C.GetNumberDimensions(); ++CDimension) // CSizes.Num() might be smaller than OutputSizes.Num()
+		{
+			const int64 CSize = C.GetSize(CDimension);
+			const int32 OutputDimension = CDimension + OutputTensor.GetNumberDimensions() - C.GetNumberDimensions();
+			const int64 OutputSize = OutputTensor.GetSize(OutputDimension);
+			if (CSize != OutputSize && CSize != 1)
+			{
+				UE_LOG(LogNeuralNetworkInference, Warning, TEXT("FGemmOperator::ConfigureOutputAndInternalVariablesAndSanityChecks(): CPadded.GetSize(%d) == 1 |"
+					" OutputTensor.GetSize(%d): %d != %d"), CDimension, OutputDimension, CSize, OutputSize);
+				return bIsLoaded;
+			}
+			// Basically
+			//  - if (C.GetNumberDimensions() == 1 && OutputTensor.GetNumberDimensions() > 1) --> CSize2D = [1, C.GetSize(0)]
+			//  - if (C.GetNumberDimensions() == 1 && OutputTensor.GetNumberDimensions() == 1) --> CSize2D = [C.GetSizes(0), 1]
+			//  - if (C.GetNumberDimensions() == 2) --> CSize2D = C.GetSizes
+			CSize2D[OutputDimension] = CSize;
+		}
+		CSizeX = CSize2D[1];
+		CSizeY = CSize2D[0];
+	}
+	// Set GEMM parameters
+	const FNeuralTensor& A = *InputTensors[0];
+	const FNeuralTensor& B = *InputTensors[1];
+	OutputRows = OutputTensor.GetSize(0);
+	OutputColumns = OutputTensor.GetSize(1);
+	AColsOrBRows = A.GetSize(bTransA ? 0 : 1);
+	if (AColsOrBRows != B.GetSize(bTransB ? 1 : 0))
+	{
+		UE_LOG(LogNeuralNetworkInference, Warning, TEXT("FGemmOperator::ConfigureOutputAndInternalVariablesAndSanityChecks(): AColsOrBRows == B.GetSize(1): %d vs. %d"),
+			AColsOrBRows, B.GetSize(bTransB ? 1 : 0));
+		return bIsLoaded;
+	}
+	AStrideX = (bTransA ? A.GetSize(1) : 1);
+	AStrideY = (bTransA ? 1 : A.GetSize(1));
+	BStrideX = (bTransB ? B.GetSize(1) : 1);
+	BStrideY = (bTransB ? 1 : B.GetSize(1));
+	OutputStride = OutputTensor.GetSize(1);
+	// Return true
+	bIsLoaded = true;
+	return bIsLoaded;
+}
+
+void FGemmOperator::ForwardCPU()
+{
+	// Sanity check
+	if (!bIsLoaded)
+	{
+		UE_LOG(LogNeuralNetworkInference, Warning, TEXT("FGemmOperator::ForwardCPU(): bIsLoaded was false."));
+		return;
+	}
+
+	// Bias
+	FNeuralTensor& OutputTensor = GetOutputTensorNoConst();
+	const TArray<FNeuralTensor*>& InputTensors = GetInputTensorsConst();
+	if (InputTensors.Num() == 3)
+	{
+		const FNeuralTensor& C = *InputTensors[2];
+		// C is scalar
+		if (C.Num() == 1)
+		{
+			const float BetaBias = Beta * C.At<float>(0);
+			OutputTensor.SetTo<float>(BetaBias);
+		}
+		// O and C both have exactly the same shape.
+		else if (OutputTensor.GetSizes() == C.GetSizes())
+		{
+			for (int64 Index = 0; Index < OutputTensor.Num(); ++Index)
+			{
+				OutputTensor.At<float>(Index) = Beta * C.At<float>(Index);
+			}
+		}
+		// O and C all have the same number of dimensions and the length of each dimensions is either a common length or C's length is 1.
+		// C has too few dimensions, and C can have its shapes prepended with a dimension of length 1 to satisfy property 2.
+		else
+		{
+			for (uint32 YOutput = 0; YOutput < OutputRows; ++YOutput)
+			{
+				for (uint32 XOutput = 0; XOutput < OutputColumns; ++XOutput)
+				{
+					const uint32 OutputIndex = YOutput * OutputStride + XOutput; // Row-major
+					const uint32 CIndex = (YOutput%CSizeY) * CSizeX + (XOutput%CSizeX);
+					OutputTensor.At<float>(OutputIndex) = Beta * C.At<float>(CIndex);
+				}
+			}
+		}
+	}
+	else
+	{
+		OutputTensor.SetTo<float>(0.f);
+	}
+
+	// Alpha x A x B
+	//		Output.noalias() += Alpha * (A * B);
+	// Column major
+	// Loop over the rows of Output
+	const FNeuralTensor& A = *InputTensors[0];
+	const FNeuralTensor& B = *InputTensors[1];
+	for (uint32 YOutput = 0u; YOutput < OutputRows; ++YOutput)
+	{
+		// Loop over the columns of Output
+		for (uint32 XOutput = 0u; XOutput < OutputColumns; ++XOutput)
+		{
+			const uint32 OutputIndex = YOutput * OutputStride + XOutput; // Row-major
+			// Update Output( YOutput,XOutput ) with the inner product of the ith row of A and the jth column of B
+			for (uint32 AColOrBRow = 0; AColOrBRow < AColsOrBRows; ++AColOrBRow)
+			{
+				// C(i,j) += A(i,p) * B(p,j);
+				const uint32 AIndex = YOutput * AStrideY + AColOrBRow * AStrideX; // Row-major (AStrideX == 1) or column-major (AStrideY == 1)
+				const uint32 BIndex = AColOrBRow * BStrideY + XOutput * BStrideX; // Row-major (BStrideX == 1) or column-major (BStrideY == 1)
+				OutputTensor.At<float>(OutputIndex) += Alpha * A.At<float>(AIndex) * B.At<float>(BIndex);
+			}
+		}
+	}
+	// Eigen equivalent of Alpha*A*B
+	// #include "Eigen/Dense"
+	//const Eigen::Map<Eigen::Matrix<float, Eigen::Dynamic, Eigen::Dynamic, Eigen::RowMajor>> A((float*)InputTensors[0]->GetData(), InputTensors[0]->GetSize(0), InputTensors[0]->GetSize(1));
+	//const Eigen::Map<Eigen::Matrix<float, Eigen::Dynamic, Eigen::Dynamic, Eigen::RowMajor>> B((float*)InputTensors[1]->GetData(), InputTensors[1]->GetSize(0), InputTensors[1]->GetSize(1));
+	//Eigen::Map<Eigen::Matrix<float, Eigen::Dynamic, Eigen::Dynamic, Eigen::RowMajor>> Output((float*)OutputTensor.GetData(), OutputTensor.GetSize(0), OutputTensor.GetSize(1));
+	//if (bTransA)
+	//	if (bTransB)
+	//		Output.noalias() += Alpha * (A.transpose() * B.transpose());
+	//	else
+	//		Output.noalias() += Alpha * (A.transpose() * B);
+	//else
+	//	if (bTransB)
+	//		Output.noalias() += Alpha * (A * B.transpose());
+	//	else
+	//		Output.noalias() += Alpha * (A * B);
+}
+
+void FGemmOperator::ForwardGPU_RenderThread(FRDGBuilder* InOutGraphBuilder)
+{
+	// Sanity checks
+	if (!FNeuralNetworkInferenceUtilsGPU::GPUSanityChecks(InOutGraphBuilder, bIsLoaded))
+	{
+		UE_LOG(LogNeuralNetworkInference, Warning, TEXT("FGemmOperator::ForwardGPU_RenderThread(): Sanity checks failed."));
+		return;
+	}
+	// Set parameters
+	FGemmCS::FParameters* Parameters = InOutGraphBuilder->AllocParameters<FGemmCS::FParameters>();
+	// Input variables
+	Parameters->Alpha = Alpha;
+	Parameters->OutputRows = OutputRows;
+	Parameters->OutputColumns = OutputColumns;
+	Parameters->AColsOrBRows = AColsOrBRows;
+	Parameters->AStrideX = AStrideX;
+	Parameters->AStrideY = AStrideY;
+	Parameters->BStrideX = BStrideX;
+	Parameters->BStrideY = BStrideY;
+	Parameters->OutputStride = OutputStride;
+	// SRV/UAV variables
+	const TArray<FNeuralTensor*>& InputTensors = GetInputTensorsConst();
+	Parameters->ASRV = InputTensors[0]->GetBufferSRVRef();
+	Parameters->BSRV = InputTensors[1]->GetBufferSRVRef();
+	Parameters->OutputUAV = GetOutputTensorNoConst().GetBufferUAVRef();
+	const EGemmMode GemmMode = (InputTensors.Num() == 2 || InputTensors[2]->Num() == 1 ? EGemmMode::CScalar :
+		InputTensors[2]->Num() > 1 ? EGemmMode::CTensor :
+		EGemmMode::MAX);
+	// No bias (i.e., 0 bias) || bias is a scalar
+	if (GemmMode == EGemmMode::CScalar)
+	{
+		Parameters->BetaTimesCScalar = (InputTensors.Num() > 2 ? Beta * InputTensors[2]->At<float>(0) : 0.f);
+	}
+	// If bias is a tensor
+	else if (GemmMode == EGemmMode::CTensor)
+	{
+		Parameters->CSRV = InputTensors[2]->GetBufferSRVRef();
+		Parameters->Beta = Beta;
+		Parameters->CSizeX = CSizeX;
+		Parameters->CSizeY = CSizeY;
+	}
+	else
+	{
+		UE_LOG(LogNeuralNetworkInference, Warning, TEXT("FGemmOperator::ForwardGPU_RenderThread(): Unknown GemmMode = %d."), (int32)GemmMode);
+		return;
+	}
+	// Set shader
+	FGemmCS::FPermutationDomain PermutationVector;
+	PermutationVector.Set<FGemmCS::FGemmMode>(GemmMode);
+	TShaderMapRef<FGemmCS> ComputeShader(GetGlobalShaderMap(GMaxRHIFeatureLevel), PermutationVector);
+	// Run shader
+	const uint32 ThreadGroupCountValueX = FMath::DivideAndRoundUp(OutputColumns, FGemmCS::THREADGROUP_SIZE_X);
+	const uint32 ThreadGroupCountValueY = FMath::DivideAndRoundUp(OutputRows, FGemmCS::THREADGROUP_SIZE_Y);
+	FComputeShaderUtils::AddPass(
+		*InOutGraphBuilder,
+		RDG_EVENT_NAME("FGemmCS() - Operator: %s, EGemmMode: %d", *Name, GemmMode),
+		ComputeShader,
+		Parameters,
+		FIntVector(ThreadGroupCountValueX, ThreadGroupCountValueY, 1));
+}
diff --git a/Engine/Plugins/Experimental/NeuralNetworkInference/Source/NeuralNetworkInference/Private/NeuralOperators/LeakyReluOperator.cpp b/Engine/Plugins/Experimental/NeuralNetworkInference/Source/NeuralNetworkInference/Private/NeuralOperators/LeakyReluOperator.cpp
new file mode 100644
--- /dev/null
+++ b/Engine/Plugins/Experimental/NeuralNetworkInference/Source/NeuralNetworkInference/Private/NeuralOperators/LeakyReluOperator.cpp
@@ -0,0 +1,28 @@
+// Copyright Epic Games, Inc. All Rights Reserved.
+
+#include "NeuralOperators/LeakyReluOperator.h"
+#include "NeuralNetworkInferenceUtils.h"
+
+
+
+/* FLeakyReluOperator structors
+ *****************************************************************************/
+
+FLeakyReluOperator::FLeakyReluOperator(const bool bIsInlinedTensor, const FNodeProto& InNodeProto)
+	: FLeakyReluOperator(bIsInlinedTensor)
+{
+	UE_LOG(LogNeuralNetworkInference, Warning, TEXT("LeakyRelu(): Constructor not tested yet."));
+	if (const FAttributeProto* AlphaAttribute = FModelProto::FindElementInArray(TEXT("Alpha"), InNodeProto.Attribute, /*bMustValueBeFound*/false))
+	{
+		Attributes[0] = AlphaAttribute->F;
+	}
+}
+
+FLeakyReluOperator::FLeakyReluOperator(const bool bIsInlinedTensor, const float InAlpha)
+	: IElementWiseOperator(TEXT("LeakyRelu"), 6, EElementWiseOperator::LeakyRelu, bIsInlinedTensor, { InAlpha })
+{
+}
+
+FLeakyReluOperator::~FLeakyReluOperator()
+{
+}
diff --git a/Engine/Plugins/Experimental/NeuralNetworkInference/Source/NeuralNetworkInference/Private/NeuralOperators/LogOperator.cpp b/Engine/Plugins/Experimental/NeuralNetworkInference/Source/NeuralNetworkInference/Private/NeuralOperators/LogOperator.cpp
new file mode 100644
--- /dev/null
+++ b/Engine/Plugins/Experimental/NeuralNetworkInference/Source/NeuralNetworkInference/Private/NeuralOperators/LogOperator.cpp
@@ -0,0 +1,17 @@
+// Copyright Epic Games, Inc. All Rights Reserved.
+
+#include "NeuralOperators/LogOperator.h"
+
+
+
+/* FLogOperator structors
+ *****************************************************************************/
+
+FLogOperator::FLogOperator(const bool bIsInlinedTensor)
+	: IElementWiseOperator(TEXT("Log"), 13, EElementWiseOperator::Log, bIsInlinedTensor)
+{
+}
+
+FLogOperator::~FLogOperator()
+{
+}
diff --git a/Engine/Plugins/Experimental/NeuralNetworkInference/Source/NeuralNetworkInference/Private/NeuralOperators/MulOperator.cpp b/Engine/Plugins/Experimental/NeuralNetworkInference/Source/NeuralNetworkInference/Private/NeuralOperators/MulOperator.cpp
new file mode 100644
--- /dev/null
+++ b/Engine/Plugins/Experimental/NeuralNetworkInference/Source/NeuralNetworkInference/Private/NeuralOperators/MulOperator.cpp
@@ -0,0 +1,17 @@
+// Copyright Epic Games, Inc. All Rights Reserved.
+
+#include "NeuralOperators/MulOperator.h"
+
+
+
+/* FMulOperator structors
+ *****************************************************************************/
+
+FMulOperator::FMulOperator(const TSet<uint32>& InPotentialInlinedTensors)
+	: IMultidirectionalBroadcastOperator(TEXT("Mul"), 13, EMultidirectionalBroadcastOperator::Mul, InPotentialInlinedTensors)
+{
+}
+
+FMulOperator::~FMulOperator()
+{
+}
diff --git a/Engine/Plugins/Experimental/NeuralNetworkInference/Source/NeuralNetworkInference/Private/NeuralOperators/MultidirectionalBroadcastOperator.cpp b/Engine/Plugins/Experimental/NeuralNetworkInference/Source/NeuralNetworkInference/Private/NeuralOperators/MultidirectionalBroadcastOperator.cpp
new file mode 100644
--- /dev/null
+++ b/Engine/Plugins/Experimental/NeuralNetworkInference/Source/NeuralNetworkInference/Private/NeuralOperators/MultidirectionalBroadcastOperator.cpp
@@ -0,0 +1,450 @@
+// Copyright Epic Games, Inc. All Rights Reserved.
+
+#include "NeuralOperators/MultidirectionalBroadcastOperator.h"
+#include "NeuralNetworkInferenceShaders/MultidirectionalBroadcastCS.h"
+#include "NeuralNetworkInferenceUtils.h"
+#include "NeuralNetworkInferenceUtilsGPU.h"
+#include "RHI.h"
+#include "Shader.h" // TShaderRef
+
+
+
+
+/* FPrivateMultidirectionalBroadcastOperator static functions
+ *****************************************************************************/
+
+class FPrivateMultidirectionalBroadcastOperator
+{
+public:
+	/**
+	 * In handles multidirectional broadcastable operators. E.g., Add, Pow, etc.
+	 */
+	static void RunMultidirectionalBroadcastElementWiseOperation(FNeuralTensor& InOutOutputTensor, const FNeuralTensor& InTensorA, const FNeuralTensor& InTensorB,
+		const TArray<uint32>& InShapeOutput, const TArray<uint32>& InShapeA, const TArray<uint32>& InShapeB, float OperationFunction(const float, const float));
+
+	/**
+	 * Eg to fill inputs.
+	 */
+	static void FillShapeMultidirectionalBroadcast(FNeuralInt64ArrayUInt32Buffer& OutShapes, const FNeuralTensor& InTensor, const uint32 InOutputNumberDimensions);
+
+	/**
+	 * Eg to fill outputs given the input shapes.
+	 */
+	static bool FillShapeMultidirectionalBroadcast(TArray<uint32>& OutShapes, const FNeuralInt64ArrayUInt32Buffer& InShapesX, const FNeuralInt64ArrayUInt32Buffer& InShapesY, const uint32 InOutputNumberDimensions);
+
+	static EMultidirectionalBroadcastShapeMode GetMultidirectionalBroadcastShapeMode(const FNeuralTensor& InX, const FNeuralTensor& InY);
+};
+
+void FPrivateMultidirectionalBroadcastOperator::RunMultidirectionalBroadcastElementWiseOperation(FNeuralTensor& InOutOutputTensor, const FNeuralTensor& InTensorA, const FNeuralTensor& InTensorB,
+	const TArray<uint32>& InShapeOutput, const TArray<uint32>& InShapeA, const TArray<uint32>& InShapeB, float OperationFunction(const float, const float))
+{
+	// 1-pass operation - This is a faster version that only works when the shapes match
+	if (InTensorA.Num() == InOutOutputTensor.Num() && InTensorA.Num() == InTensorB.Num())
+	{
+		for (uint32 IndexOutput1D = 0; IndexOutput1D < InTensorA.Num(); ++IndexOutput1D)
+		{
+			InOutOutputTensor.At<float>(IndexOutput1D) = OperationFunction(InTensorA.At<float>(IndexOutput1D), InTensorB.At<float>(IndexOutput1D));
+		}
+	}
+	// B is scalar
+	else if (InTensorB.Num() == 1)
+	{
+		const float BValue = InTensorB.At<float>(0);
+		for (uint32 IndexOutput1D = 0; IndexOutput1D < InOutOutputTensor.Num(); ++IndexOutput1D)
+		{
+			InOutOutputTensor.At<float>(IndexOutput1D) = OperationFunction(InTensorA.At<float>(IndexOutput1D), BValue);
+		}
+	}
+	// A is scalar
+	else if (InTensorA.Num() == 1)
+	{
+		const float AValue = InTensorA.At<float>(0);
+		for (uint32 IndexOutput1D = 0; IndexOutput1D < InOutOutputTensor.Num(); ++IndexOutput1D)
+		{
+			InOutOutputTensor.At<float>(IndexOutput1D) = OperationFunction(AValue, InTensorB.At<float>(IndexOutput1D));
+		}
+	}
+	// InTensorA and/or InTensorB and/or Output does not match on shape. Slower version that would work for any case.
+	else
+	{
+		if (InShapeOutput.Num() != InShapeA.Num() || InShapeOutput.Num() != InShapeB.Num())
+		{
+			UE_LOG(LogNeuralNetworkInference, Warning, TEXT("IMultidirectionalBroadcastOperator::RunMultidirectionalBroadcastElementWiseOperation():"
+				" InShapeOutput.Num() == InShapeA.Num() == InShapeB.Num() failed: %d vs. %d vs. %d."), InShapeOutput.Num(), InShapeA.Num(), InShapeB.Num());
+			return;
+		}
+		TArray<uint32> CurrentIndexArray;
+		CurrentIndexArray.Init(0, InShapeOutput.Num());
+		for (uint32 IndexOutput1D = 0; IndexOutput1D < InOutOutputTensor.Num(); ++IndexOutput1D)
+		{
+			// Estimate right 1-D indexes
+			int32 IndexA1D = 0;
+			int32 IndexB1D = 0;
+			for (int32 ShapeIndex = 0; ShapeIndex < InShapeOutput.Num(); ++ShapeIndex)
+			{
+				//IndexOutput1D = IndexOutput1D * InShapeOutput[ShapeIndex] + CurrentIndexArray[ShapeIndex]; // IndexOutput1D = [[[[...]P_3M + P3]P_2M + P_2]P_1M + P_1]P_0M + P0
+				IndexA1D = IndexA1D * InShapeA[ShapeIndex] + (CurrentIndexArray[ShapeIndex] % InShapeA[ShapeIndex]); // IndexA1D = [[[[...]P_3M + P3']P_2M + P_2']P_1M + P_1']P_0M + P0', P_i' = P_i%P_iM
+				IndexB1D = IndexB1D * InShapeB[ShapeIndex] + (CurrentIndexArray[ShapeIndex] % InShapeB[ShapeIndex]); // IndexB1D = [[[[...]P_3M + P3']P_2M + P_2']P_1M + P_1']P_0M + P0', P_i' = P_i%P_iM
+			}
+			// Perform operation on the right indexes
+			InOutOutputTensor.At<float>(IndexOutput1D) = OperationFunction(InTensorA.At<float>(IndexA1D), InTensorB.At<float>(IndexB1D));
+			// Update CurrentIndexArray (for ShapesOutput = [5 2 3]):
+			//    [0 0 0] -> [0 0 1] -> [0 0 2] -> [0 1 0] -> [0 1 1] -> [0 1 2] -> [1 0 0] -> [1 0 1] -> [1 1 2] -> [2 0 0] -> ...
+			for (int32 CurrentIndexArrayIndex = CurrentIndexArray.Num() - 1; CurrentIndexArrayIndex > -1; --CurrentIndexArrayIndex)
+			{
+				if (CurrentIndexArray[CurrentIndexArrayIndex] + 1 == InShapeOutput[CurrentIndexArrayIndex])
+				{
+					CurrentIndexArray[CurrentIndexArrayIndex] = 0;
+				}
+				else
+				{
+					++CurrentIndexArray[CurrentIndexArrayIndex];
+					break;
+				}
+			}
+		}
+	}
+}
+
+void FPrivateMultidirectionalBroadcastOperator::FillShapeMultidirectionalBroadcast(FNeuralInt64ArrayUInt32Buffer& OutShapes, const FNeuralTensor& InTensor, const uint32 InOutputNumberDimensions)
+{
+	TArray<uint32> ShapeArray;
+	ShapeArray.SetNumUninitialized(InOutputNumberDimensions);
+	for (int32 DimensionIndex = 0; DimensionIndex < ShapeArray.Num(); ++DimensionIndex)
+	{
+		// 1s at the beginning
+		if (InTensor.GetNumberDimensions() < ShapeArray.Num() - DimensionIndex)
+		{
+			ShapeArray[DimensionIndex] = 1;
+		}
+		// Rest
+		else
+		{
+			ShapeArray[DimensionIndex] = InTensor.GetSize(DimensionIndex + InTensor.GetNumberDimensions() - ShapeArray.Num());
+		}
+	}
+	OutShapes.Move(ShapeArray, /*bShouldCreate64BitVersion*/false);
+}
+
+bool FPrivateMultidirectionalBroadcastOperator::FillShapeMultidirectionalBroadcast(TArray<uint32>& OutShapes, const FNeuralInt64ArrayUInt32Buffer& InShapesX, const FNeuralInt64ArrayUInt32Buffer& InShapesY,
+	const uint32 InOutputNumberDimensions)
+{
+	// Sanity checks - All shapes must have same size
+	if (InShapesX.Num() != InOutputNumberDimensions || InShapesY.Num() != InOutputNumberDimensions)
+	{
+		UE_LOG(LogNeuralNetworkInference, Warning, TEXT("IMultidirectionalBroadcastOperator::FillShapeMultidirectionalBroadcast():"
+			" InShapesX.Num() == InShapesY.Num() == InOutputNumberDimensions failed: %d vs. %d vs. %d"), InShapesX.Num(), InShapesY.Num(), InOutputNumberDimensions);
+		return false;
+	}
+	// Fill OutShapes
+	OutShapes.SetNumUninitialized(InOutputNumberDimensions);
+	for (int32 DimensionIndex = 0; DimensionIndex < OutShapes.Num(); ++DimensionIndex)
+	{
+		OutShapes[DimensionIndex] = FMath::Max(InShapesX.GetUInt32()[DimensionIndex], InShapesY.GetUInt32()[DimensionIndex]);
+	}
+	return true;
+}
+
+EMultidirectionalBroadcastShapeMode FPrivateMultidirectionalBroadcastOperator::GetMultidirectionalBroadcastShapeMode(const FNeuralTensor& InX, const FNeuralTensor& InY)
+{
+	// 1-pass operation - This is a faster version that only works when the shapes match
+	if (InX.GetSizes() == InY.GetSizes()) // This is different than InX.Num() == InY.Num(), eg if X={1,3} and Y={3,1}
+	{
+		return EMultidirectionalBroadcastShapeMode::ElementWise;
+	}
+	else if (InX.Num() == 1)
+	{
+		return EMultidirectionalBroadcastShapeMode::AScalar;
+	}
+	else if (InY.Num() == 1)
+	{
+		return EMultidirectionalBroadcastShapeMode::BScalar;
+	}
+	// InTensorA and/or InTensorB and/or Output does not match on shape. Slower version that would work for any case
+	else
+	{
+		return EMultidirectionalBroadcastShapeMode::MultidirectionalBroadcast;
+	}
+}
+
+
+
+/* IMultidirectionalBroadcastOperator structors
+ *****************************************************************************/
+
+IMultidirectionalBroadcastOperator::IMultidirectionalBroadcastOperator(const FString& InName, const int32 InVersion,
+	const EMultidirectionalBroadcastOperator InMultidirectionalBroadcastOperator, const TSet<uint32>& InPotentialInlinedTensors)
+	: FNeuralOperator(InName, InVersion, InPotentialInlinedTensors)
+	, MultidirectionalBroadcastOperator(InMultidirectionalBroadcastOperator)
+	, MultidirectionalBroadcastShapeMode(EMultidirectionalBroadcastShapeMode::MAX)
+{
+}
+
+IMultidirectionalBroadcastOperator::~IMultidirectionalBroadcastOperator()
+{
+}
+
+
+
+/* IMultidirectionalBroadcastOperator public functions
+ *****************************************************************************/
+
+bool IMultidirectionalBroadcastOperator::ConfigureOutputAndInternalVariablesAndSanityChecks()
+{
+	bIsLoaded = false;
+	// Input sanity checks
+	const TArray<FNeuralTensor*>& InputTensors = GetInputTensorsConst();
+	if (!FNeuralNetworkInferenceUtils::SizeSanityChecks(InputTensors, 2))
+	{
+		return bIsLoaded;
+	}
+	if (ShapesOutput.Num() != ShapesX.Num() || ShapesOutput.Num() != ShapesY.Num())
+	{
+		UE_LOG(LogNeuralNetworkInference, Warning, TEXT("IMultidirectionalBroadcastOperator::ConfigureOutputAndInternalVariablesAndSanityChecks():"
+			" ShapesOutput.Num() == ShapesX.Num() == ShapesY.Num() failed: %d vs. %d vs. %d."), ShapesOutput.Num(), ShapesX.Num(), ShapesY.Num());
+		return bIsLoaded;
+	}
+	if (InlinedTensor < -1)
+	{
+		UE_LOG(LogNeuralNetworkInference, Warning, TEXT("IMultidirectionalBroadcastOperator::ConfigureOutputAndInternalVariablesAndSanityChecks(): Invalid InlinedTensor = %d."), InlinedTensor);
+		return bIsLoaded;
+	}
+	// Output (if not inlined)
+	if (InlinedTensor < 0)
+	{
+		TArray<FNeuralTensor*>& OutputTensors = GetOutputTensorsNoConst();
+		// Sanity checks
+		if (!FNeuralNetworkInferenceUtils::SizeSanityChecks(OutputTensors, 1))
+		{
+			return bIsLoaded;
+		}
+		// Set and configure output tensor
+		FNeuralTensor& OutputTensor = *OutputTensors.Last();
+		OutputTensor.SetNumUninitialized(ShapesOutput.GetInt64(), InputTensors[0]->GetDataType());
+	}
+	// ShapesOutput < MAX_NUMBER_DIMENSIONS
+	if ((uint32)ShapesOutput.Num() > FMultidirectionalBroadcastCS::MAX_NUMBER_DIMENSIONS)
+	{
+		UE_LOG(LogNeuralNetworkInference, Warning, TEXT("IMultidirectionalBroadcastOperator::ConfigureOutputAndInternalVariablesAndSanityChecks(): Only implemented for"
+			" ShapesOutput.Num() <= FMultidirectionalBroadcastCS::MAX_NUMBER_DIMENSIONS: %d vs. %d. Contact us if you need more dimensions."),
+			ShapesOutput.Num(), FMultidirectionalBroadcastCS::MAX_NUMBER_DIMENSIONS);
+		return bIsLoaded;
+	}
+	// Return true
+	bIsLoaded = true;
+	return bIsLoaded;
+}
+
+void IMultidirectionalBroadcastOperator::ToGPU_RenderThread()
+{
+	// Sanity check
+	if (ShapesX.IsReadBufferValid() || ShapesY.IsReadBufferValid() || ShapesOutput.IsReadBufferValid())
+	{
+		UE_LOG(LogNeuralNetworkInference, Warning, TEXT("IMultidirectionalBroadcastOperator::ToGPU_RenderThread():"
+			" ShapesX, ShapesY and/or ShapesOutput FReadBuffers were not nullptr."));
+	}
+	// Set MultidirectionalBroadcastShapeMode
+	const TArray<FNeuralTensor*>& InputTensors = GetInputTensorsConst();
+	MultidirectionalBroadcastShapeMode = FPrivateMultidirectionalBroadcastOperator::GetMultidirectionalBroadcastShapeMode(*InputTensors[0], *InputTensors[1]);
+	// Memory to GPU
+	if (MultidirectionalBroadcastShapeMode == EMultidirectionalBroadcastShapeMode::MultidirectionalBroadcast)
+	{
+		ShapesX.CreateAndLoadSRVBuffer(TEXT("MultidirectionalBroadcast_ShapesX"));
+		ShapesY.CreateAndLoadSRVBuffer(TEXT("MultidirectionalBroadcast_ShapesY"));
+		ShapesOutput.CreateAndLoadSRVBuffer(TEXT("MultidirectionalBroadcast_ShapesOutput"));
+	}
+}
+
+void IMultidirectionalBroadcastOperator::ForwardGPU_RenderThread(FRDGBuilder* InOutGraphBuilder)
+{
+	// Sanity checks
+	if (!FNeuralNetworkInferenceUtilsGPU::GPUSanityChecks(InOutGraphBuilder, bIsLoaded))
+	{
+		UE_LOG(LogNeuralNetworkInference, Warning, TEXT("IMultidirectionalBroadcastOperator::ForwardGPU_RenderThread():"
+			" IsInRenderingThread() = %d != true || InOutGraphBuilder is nullptr."), IsInRenderingThread());
+		return;
+	}
+	// Set parameters
+	const TArray<FNeuralTensor*>& InputTensors = GetInputTensorsConst();
+	FMultidirectionalBroadcastCS::FParameters* Parameters = InOutGraphBuilder->AllocParameters<FMultidirectionalBroadcastCS::FParameters>();
+	const uint32 Volume = (uint32)(InlinedTensor < 0 ? GetOutputTensorConst().Num() : InputTensors[InlinedTensor]->Num());
+	Parameters->TensorSize = Volume;
+	// Set MultidirectionalBroadcastInlinedMode
+	EMultidirectionalBroadcastInlinedMode MultidirectionalBroadcastInlinedMode;
+	const FNeuralTensor& X = *InputTensors[0];
+	const FNeuralTensor& Y = *InputTensors[1];
+	// Not inlined
+	if (InlinedTensor < 0)
+	{
+		MultidirectionalBroadcastInlinedMode = EMultidirectionalBroadcastInlinedMode::None;
+		Parameters->ASRV = X.GetBufferSRVRef();
+		Parameters->BSRV = Y.GetBufferSRVRef();
+		Parameters->OutputUAV = GetOutputTensorNoConst().GetBufferUAVRef();
+		if (MultidirectionalBroadcastShapeMode == EMultidirectionalBroadcastShapeMode::MultidirectionalBroadcast)
+		{
+			Parameters->ShapeDimensions = ShapesOutput.Num();
+			Parameters->ShapeOutput = ShapesOutput.GetUInt32SRV();
+			Parameters->ShapeA = ShapesX.GetUInt32SRV();
+			Parameters->ShapeB = ShapesY.GetUInt32SRV();
+		}
+	}
+	// Inlined
+	else
+	{
+		TArray<FNeuralTensor*>& InputTensorsNotConst = GetInputTensorsNoConst();
+		// X Inlined
+		if (InlinedTensor == 0)
+		{
+			MultidirectionalBroadcastInlinedMode = EMultidirectionalBroadcastInlinedMode::A;
+			Parameters->BSRV = Y.GetBufferSRVRef();
+			Parameters->OutputUAV = InputTensorsNotConst[0]->GetBufferUAVRef();
+			if (MultidirectionalBroadcastShapeMode == EMultidirectionalBroadcastShapeMode::MultidirectionalBroadcast)
+			{
+				Parameters->ShapeDimensions = ShapesOutput.Num();
+				Parameters->ShapeOutput = ShapesOutput.GetUInt32SRV();
+				Parameters->ShapeB = ShapesY.GetUInt32SRV();
+			}
+		}
+		// Y Inlined
+		else if (InlinedTensor == 1)
+		{
+			MultidirectionalBroadcastInlinedMode = EMultidirectionalBroadcastInlinedMode::B;
+			Parameters->ASRV = X.GetBufferSRVRef();
+			Parameters->OutputUAV = InputTensorsNotConst[1]->GetBufferUAVRef();
+			if (MultidirectionalBroadcastShapeMode == EMultidirectionalBroadcastShapeMode::MultidirectionalBroadcast)
+			{
+				Parameters->ShapeDimensions = ShapesOutput.Num();
+				Parameters->ShapeOutput = ShapesOutput.GetUInt32SRV();
+				Parameters->ShapeA = ShapesX.GetUInt32SRV();
+			}
+		}
+		// Unknown case
+		else
+		{
+			MultidirectionalBroadcastInlinedMode = EMultidirectionalBroadcastInlinedMode::None; // To avoid unset warning
+			UE_LOG(LogNeuralNetworkInference, Warning, TEXT("IMultidirectionalBroadcastOperator::ForwardGPU_RenderThread(): Unknown InlinedTensor = %d."), InlinedTensor);
+			return;
+		}
+	}
+	// Set shader
+	FMultidirectionalBroadcastCS::FPermutationDomain PermutationVector;
+	PermutationVector.Set<FMultidirectionalBroadcastCS::FShaderType>(MultidirectionalBroadcastOperator);
+	PermutationVector.Set<FMultidirectionalBroadcastCS::FInlinedMode>(MultidirectionalBroadcastInlinedMode);
+	PermutationVector.Set<FMultidirectionalBroadcastCS::FShapeMode>(MultidirectionalBroadcastShapeMode);
+	TShaderMapRef<FMultidirectionalBroadcastCS> ComputeShader(GetGlobalShaderMap(GMaxRHIFeatureLevel), PermutationVector);
+	// Run shader
+	const uint32 ThreadGroupCountValueX = FMath::DivideAndRoundUp(Volume, FMultidirectionalBroadcastCS::THREADGROUP_SIZE_X);
+	FComputeShaderUtils::AddPass(
+		*InOutGraphBuilder,
+		RDG_EVENT_NAME("FMultidirectionalBroadcastCS() - Operator: %s, InlinedTensor: %d", *Name, InlinedTensor),
+		ComputeShader,
+		Parameters,
+		FIntVector(ThreadGroupCountValueX, 1, 1));
+}
+
+
+
+/* IMultidirectionalBroadcastOperator protected functions
+ *****************************************************************************/
+
+void IMultidirectionalBroadcastOperator::ForwardCPUWithFunction(float InOperatorFunction(const float, const float))
+{
+	// Sanity checks
+	if (!bIsLoaded)
+	{
+		UE_LOG(LogNeuralNetworkInference, Warning, TEXT("IMultidirectionalBroadcastOperator::ForwardCPUWithFunction(): bIsLoaded was false."));
+		return;
+	}
+	if (InlinedTensor < -1)
+	{
+		UE_LOG(LogNeuralNetworkInference, Warning, TEXT("IMultidirectionalBroadcastOperator::ForwardCPUWithFunction(): Invalid InlinedTensor = %d."), InlinedTensor);
+		return;
+	}
+	// Not inlined layer (output is different than inputs)
+	if (InlinedTensor < 0)
+	{
+		const TArray<FNeuralTensor*>& InputTensors = GetInputTensorsConst();
+		// Running operator in a multidirectional broadcastable way
+		FPrivateMultidirectionalBroadcastOperator::RunMultidirectionalBroadcastElementWiseOperation(GetOutputTensorNoConst(),
+			*InputTensors[0], *InputTensors[1], ShapesOutput.GetUInt32(), ShapesX.GetUInt32(), ShapesY.GetUInt32(), InOperatorFunction);
+	}
+	// Inlined layer
+	else
+	{
+		TArray<FNeuralTensor*>& InputTensors = GetInputTensorsNoConst();
+		FNeuralTensor& X = *InputTensors[0];
+		FNeuralTensor& Y = *InputTensors[1];
+		// Running operator in a multidirectional broadcastable way
+		FPrivateMultidirectionalBroadcastOperator::RunMultidirectionalBroadcastElementWiseOperation((InlinedTensor == 0 ? X : Y),
+			X, Y, ShapesOutput.GetUInt32(), ShapesX.GetUInt32(), ShapesY.GetUInt32(), InOperatorFunction);
+	}
+}
+
+bool IMultidirectionalBroadcastOperator::EstimateInlinedTensorFromPotentialOnes()
+{
+	const TArray<FNeuralTensor*>& InputTensors = GetInputTensorsConst();
+	if (InputTensors.Num() != 2)
+	{
+		UE_LOG(LogNeuralNetworkInference, Warning, TEXT("IMultidirectionalBroadcastOperator::EstimateInlinedTensorFromPotentialOnes(): InputTensors.Num() = %d != 2."), InputTensors.Num());
+		return false;
+	}
+
+	// Create ShapesX, ShapesY, and ShapesOutput GetInt64()/GetUInt32()
+	const FNeuralTensor& X = *InputTensors[0];
+	const FNeuralTensor& Y = *InputTensors[1];
+	const uint32 OutputNumberDimensions = FMath::Max(X.GetNumberDimensions(), Y.GetNumberDimensions());
+	FPrivateMultidirectionalBroadcastOperator::FillShapeMultidirectionalBroadcast(ShapesX, X, OutputNumberDimensions);
+	FPrivateMultidirectionalBroadcastOperator::FillShapeMultidirectionalBroadcast(ShapesY, Y, OutputNumberDimensions);
+	TArray<uint32> OutputShapeTemporary;
+	if (!FPrivateMultidirectionalBroadcastOperator::FillShapeMultidirectionalBroadcast(OutputShapeTemporary, ShapesX, ShapesY, OutputNumberDimensions))
+	{
+		return false;
+	}
+	ShapesOutput.Move(OutputShapeTemporary, /*bShouldCreate64BitVersion*/ true);
+
+	// Can layer be inlined?
+	if (PotentiallyInlinedTensors.Num() > 0)
+	{
+		// Size(X) == Size(Output)
+		if (X.GetSizes() == ShapesOutput.GetInt64() && PotentiallyInlinedTensors.Contains(0))
+		{
+			InlinedTensor = 0;
+		}
+		// Size(Y) == Size(Output)
+		else if (Y.GetSizes() == ShapesOutput.GetInt64() && PotentiallyInlinedTensors.Contains(1))
+		{
+			InlinedTensor = 1;
+		}
+		// No same size than output
+		else
+		{
+			InlinedTensor = -1;
+		}
+	}
+
+	// Reset shape buffers
+	if (ShapesOutput.IsReadBufferValid())
+	{
+		ShapesX.ReleaseReadBuffer();
+		ShapesY.ReleaseReadBuffer();
+		ShapesOutput.ReleaseReadBuffer();
+	}
+	ShapesX.ResetReadBuffer();
+	ShapesY.ResetReadBuffer();
+	ShapesOutput.ResetReadBuffer();
+	return true;
+}
+
+void IMultidirectionalBroadcastOperator::ShapesToGPU()
+{
+	// Sanity check
+	if (ShapesX.IsReadBufferValid() || ShapesY.IsReadBufferValid() || ShapesOutput.IsReadBufferValid())
+	{
+		UE_LOG(LogNeuralNetworkInference, Warning, TEXT("IMultidirectionalBroadcastOperator::ShapesToGPU():"
+			" ShapesX, ShapesY and/or ShapesOutput FReadBuffers were not nullptr."));
+	}
+	// Create GPU memory for it
+	ShapesX.CreateAndLoadSRVBuffer(TEXT("MultidirectionalBroadcast_ShapesX"));
+	ShapesY.CreateAndLoadSRVBuffer(TEXT("MultidirectionalBroadcast_ShapesY"));
+	ShapesOutput.CreateAndLoadSRVBuffer(TEXT("MultidirectionalBroadcast_ShapesOutput"));
+}
diff --git a/Engine/Plugins/Experimental/NeuralNetworkInference/Source/NeuralNetworkInference/Private/NeuralOperators/NegOperator.cpp b/Engine/Plugins/Experimental/NeuralNetworkInference/Source/NeuralNetworkInference/Private/NeuralOperators/NegOperator.cpp
new file mode 100644
--- /dev/null
+++ b/Engine/Plugins/Experimental/NeuralNetworkInference/Source/NeuralNetworkInference/Private/NeuralOperators/NegOperator.cpp
@@ -0,0 +1,17 @@
+// Copyright Epic Games, Inc. All Rights Reserved.
+
+#include "NeuralOperators/NegOperator.h"
+
+
+
+/* FNegOperator structors
+ *****************************************************************************/
+
+FNegOperator::FNegOperator(const bool bIsInlinedTensor)
+	: IElementWiseOperator(TEXT("Neg"), 13, EElementWiseOperator::Neg, bIsInlinedTensor)
+{
+}
+
+FNegOperator::~FNegOperator()
+{
+}
diff --git a/Engine/Plugins/Experimental/NeuralNetworkInference/Source/NeuralNetworkInference/Private/NeuralOperators/PowOperator.cpp b/Engine/Plugins/Experimental/NeuralNetworkInference/Source/NeuralNetworkInference/Private/NeuralOperators/PowOperator.cpp
new file mode 100644
--- /dev/null
+++ b/Engine/Plugins/Experimental/NeuralNetworkInference/Source/NeuralNetworkInference/Private/NeuralOperators/PowOperator.cpp
@@ -0,0 +1,17 @@
+// Copyright Epic Games, Inc. All Rights Reserved.
+
+#include "NeuralOperators/PowOperator.h"
+
+
+
+/* FPowOperator structors
+ *****************************************************************************/
+
+FPowOperator::FPowOperator(const TSet<uint32>& InPotentialInlinedTensors)
+	: IMultidirectionalBroadcastOperator(TEXT("Pow"), 13, EMultidirectionalBroadcastOperator::Pow, InPotentialInlinedTensors)
+{
+}
+
+FPowOperator::~FPowOperator()
+{
+}
diff --git a/Engine/Plugins/Experimental/NeuralNetworkInference/Source/NeuralNetworkInference/Private/NeuralOperators/ReciprocalOperator.cpp b/Engine/Plugins/Experimental/NeuralNetworkInference/Source/NeuralNetworkInference/Private/NeuralOperators/ReciprocalOperator.cpp
new file mode 100644
--- /dev/null
+++ b/Engine/Plugins/Experimental/NeuralNetworkInference/Source/NeuralNetworkInference/Private/NeuralOperators/ReciprocalOperator.cpp
@@ -0,0 +1,17 @@
+// Copyright Epic Games, Inc. All Rights Reserved.
+
+#include "NeuralOperators/ReciprocalOperator.h"
+
+
+
+/* FReciprocalOperator structors
+ *****************************************************************************/
+
+FReciprocalOperator::FReciprocalOperator(const bool bIsInlinedTensor)
+	: IElementWiseOperator(TEXT("Reciprocal"), 13, EElementWiseOperator::Reciprocal, bIsInlinedTensor)
+{
+}
+
+FReciprocalOperator::~FReciprocalOperator()
+{
+}
diff --git a/Engine/Plugins/Experimental/NeuralNetworkInference/Source/NeuralNetworkInference/Private/NeuralOperators/ReluOperator.cpp b/Engine/Plugins/Experimental/NeuralNetworkInference/Source/NeuralNetworkInference/Private/NeuralOperators/ReluOperator.cpp
new file mode 100644
--- /dev/null
+++ b/Engine/Plugins/Experimental/NeuralNetworkInference/Source/NeuralNetworkInference/Private/NeuralOperators/ReluOperator.cpp
@@ -0,0 +1,17 @@
+// Copyright Epic Games, Inc. All Rights Reserved.
+
+#include "NeuralOperators/ReluOperator.h"
+
+
+
+/* FReluOperator structors
+ *****************************************************************************/
+
+FReluOperator::FReluOperator(const bool bIsInlinedTensor)
+	: IElementWiseOperator(TEXT("Relu"), 13, EElementWiseOperator::Relu, bIsInlinedTensor)
+{
+}
+
+FReluOperator::~FReluOperator()
+{
+}
diff --git a/Engine/Plugins/Experimental/NeuralNetworkInference/Source/NeuralNetworkInference/Private/NeuralOperators/ReshapeOperator.cpp b/Engine/Plugins/Experimental/NeuralNetworkInference/Source/NeuralNetworkInference/Private/NeuralOperators/ReshapeOperator.cpp
new file mode 100644
--- /dev/null
+++ b/Engine/Plugins/Experimental/NeuralNetworkInference/Source/NeuralNetworkInference/Private/NeuralOperators/ReshapeOperator.cpp
@@ -0,0 +1,177 @@
+// Copyright Epic Games, Inc. All Rights Reserved.
+
+#include "NeuralOperators/ReshapeOperator.h"
+#include "NeuralNetworkInferenceShaders/CopyCS.h"
+#include "NeuralNetworkInferenceUtils.h"
+#include "NeuralNetworkInferenceUtilsGPU.h"
+
+
+
+/* FReshapeOperator structors
+ *****************************************************************************/
+
+FReshapeOperator::FReshapeOperator(const bool bIsInlinedTensor, const FNodeProto& InNodeProto)
+	: FReshapeOperator(bIsInlinedTensor)
+{
+	if (const FAttributeProto* MomentumAttribute = FModelProto::FindElementInArray(TEXT("AllowZero"), InNodeProto.Attribute, /*bMustValueBeFound*/false))
+	{
+		AllowZero = MomentumAttribute->I;
+	}
+}
+
+FReshapeOperator::FReshapeOperator(const bool bIsInlinedTensor, const int64 InAllowZero)
+	: FNeuralOperator(TEXT("Reshape"), 14, (bIsInlinedTensor ? 0 : -1))
+	, AllowZero(InAllowZero)
+{
+}
+
+FReshapeOperator::~FReshapeOperator()
+{
+}
+
+
+
+/* FReshapeOperator public functions
+ *****************************************************************************/
+
+bool FReshapeOperator::ConfigureOutputAndInternalVariablesAndSanityChecks()
+{
+	bIsLoaded = false;
+	// Output initialization and size sanity check
+	if (InlinedTensor < 0)
+	{
+		const TArray<FNeuralTensor*>& InputTensors = GetInputTensorsConst();
+		const FNeuralTensor& InputTensor = *InputTensors[0];
+		const FNeuralTensor& ShapeTensor = *InputTensors[1];
+		FNeuralTensor& OutputTensor = GetOutputTensorNoConst();
+		// Initialize OutputTensor
+		const TArray<int64> OutputShape = FReshapeOperator::Remove0sAndMinus1sFromShape(ShapeTensor, InputTensor);
+		OutputTensor.SetNumUninitialized(OutputShape, InputTensor.GetDataType());
+	}
+	// Return true
+	bIsLoaded = true;
+	return bIsLoaded;
+}
+
+bool FReshapeOperator::SetWhetherInputTensorsNeedTransferToGPUForGPUMode()
+{
+	// Enable GPU mode only on the first input tensor
+	const TArray<FNeuralTensor*>& InputTensors = GetInputTensorsConst();
+	InputTensors[0]->SetEnableGPU(true);
+	return true;
+}
+
+void FReshapeOperator::ForwardCPU()
+{
+	// Sanity check
+	if (!bIsLoaded)
+	{
+		UE_LOG(LogNeuralNetworkInference, Warning, TEXT("FReshapeOperator::ForwardCPU(): bIsLoaded was false."));
+		return;
+	}
+	// Reshape is 0-cost if inlined, otherwise, just copy memory from input to output
+	const TArray<FNeuralTensor*> InputTensors = GetInputTensorsConst();
+	const FNeuralTensor& InputTensor = *InputTensors[0];
+	const FNeuralTensor& ShapeTensor = *InputTensors[1];
+	FNeuralTensor& OutputTensor = (InlinedTensor < 0 ? GetOutputTensorNoConst() : *GetInputTensorsNoConst()[0]);
+	// Update InputSizeIfInlined (if inlined)
+	if (InlinedTensor > -1)
+	{
+		InputSizeIfInlined = InputTensor.GetSizes();
+	}
+	// Copy memory (if not inlined)
+	else
+	{
+		OutputTensor.SetFromUnderlyingUInt8ArrayCopy(InputTensor.GetUnderlyingUInt8ArrayRef());
+	}
+	// Reshape output
+	TArray<int64> OutputShape = FReshapeOperator::Remove0sAndMinus1sFromShape(ShapeTensor, InputTensor);
+	OutputTensor.ReshapeMove(OutputShape);
+}
+
+void FReshapeOperator::ForwardGPU_RenderThread(FRDGBuilder* InOutGraphBuilder)
+{
+	// Sanity checks
+	if (!FNeuralNetworkInferenceUtilsGPU::GPUSanityChecks(InOutGraphBuilder, bIsLoaded))
+	{
+		UE_LOG(LogNeuralNetworkInference, Warning, TEXT("FReshapeOperator::ForwardGPU_RenderThread(): Sanity checks failed."));
+		return;
+	}
+	// Reshape is 0-cost if inlined, otherwise, just copy memory from input to output
+	const TArray<FNeuralTensor*> InputTensors = GetInputTensorsConst();
+	const FNeuralTensor& InputTensor = *InputTensors[0];
+	const FNeuralTensor& ShapeTensor = *InputTensors[1];
+	FNeuralTensor& OutputTensor = (InlinedTensor < 0 ? GetOutputTensorNoConst() : *GetInputTensorsNoConst()[0]);
+	// If input is empty, it means it has a GetSizes() similar to [a, b, 0, c, ...]
+	if (InputTensor.IsEmpty())
+	{
+		return;
+	}
+	// Update InputSizeIfInlined (if inlined)
+	if (InlinedTensor > -1)
+	{
+		InputSizeIfInlined = InputTensor.GetSizes();
+	}
+	// Copy memory (if not inlined)
+	else
+	{
+		// Set parameters
+		FCopyCS::FParameters* Parameters = InOutGraphBuilder->AllocParameters<FCopyCS::FParameters>();
+		Parameters->Num = OutputTensor.Num();
+		// SRV/UAV variables
+		Parameters->InputSRV = InputTensor.GetBufferSRVRef();
+		Parameters->OutputUAV = OutputTensor.GetBufferUAVRef();
+		// Set shader
+		TShaderMapRef<FCopyCS> ComputeShader(GetGlobalShaderMap(GMaxRHIFeatureLevel));
+		// Run shader
+		const uint32 ThreadGroupCountValueX = FMath::DivideAndRoundUp(Parameters->Num, FCopyCS::THREADGROUP_SIZE_X);
+		FComputeShaderUtils::AddPass(
+			*InOutGraphBuilder,
+			RDG_EVENT_NAME("FReshapeOperator-CopyCS()"),
+			ComputeShader,
+			Parameters,
+			FIntVector(ThreadGroupCountValueX, 1, 1));
+	}
+	// Reshape output
+	TArray<int64> OutputShape = FReshapeOperator::Remove0sAndMinus1sFromShape(ShapeTensor, InputTensor);
+	OutputTensor.ReshapeMove(OutputShape);
+}
+
+void FReshapeOperator::PostForwardCPU()
+{
+	if (NeedsPostForward())
+	{
+		FNeuralTensor& InputTensor = *GetInputTensorsNoConst()[0];
+		InputTensor.Reshape(InputSizeIfInlined);
+	}
+}
+
+TArray<int64> FReshapeOperator::Remove0sAndMinus1sFromShape(const FNeuralTensor& InShapeTensor, const FNeuralTensor& InTensor) const
+{
+	TArray<int64> OutputShape = InShapeTensor.GetArrayCopy<int64>();
+	const TArray<int64>& CurrentShape = InTensor.GetSizes();
+	const int64 FinalVolume = InTensor.Num();
+
+	int64 Minus1Index = -1;
+	for (int64 ShapeIndex = 0; ShapeIndex < OutputShape.Num(); ++ShapeIndex)
+	{
+		// 0 = Copy from input
+		if (OutputShape[ShapeIndex] == 0 && !AllowZero)
+		{
+			OutputShape[ShapeIndex] = CurrentShape[ShapeIndex];
+		}
+		// -1 = Infer from remaining dimensions
+		else if (OutputShape[ShapeIndex] == -1)
+		{
+			Minus1Index = ShapeIndex;
+			OutputShape[ShapeIndex] = 1;
+		}
+	}
+	// Infer from remaining dimensions - All 0s must have been already converted into the right dimensions, and the only -1 turned into 1 for Product() to work properly
+	if (Minus1Index > -1)
+	{
+		OutputShape[Minus1Index] = FinalVolume / FNeuralNetworkInferenceUtils::Product<int64>(OutputShape);
+	}
+	// Return cleaned OutputShape (without 0s or -1s)
+	return OutputShape;
+}
diff --git a/Engine/Plugins/Experimental/NeuralNetworkInference/Source/NeuralNetworkInference/Private/NeuralOperators/RoundOperator.cpp b/Engine/Plugins/Experimental/NeuralNetworkInference/Source/NeuralNetworkInference/Private/NeuralOperators/RoundOperator.cpp
new file mode 100644
--- /dev/null
+++ b/Engine/Plugins/Experimental/NeuralNetworkInference/Source/NeuralNetworkInference/Private/NeuralOperators/RoundOperator.cpp
@@ -0,0 +1,17 @@
+// Copyright Epic Games, Inc. All Rights Reserved.
+
+#include "NeuralOperators/RoundOperator.h"
+
+
+
+/* FRoundOperator structors
+ *****************************************************************************/
+
+FRoundOperator::FRoundOperator(const bool bIsInlinedTensor)
+	: IElementWiseOperator(TEXT("Round"), 11, EElementWiseOperator::Round, bIsInlinedTensor)
+{
+}
+
+FRoundOperator::~FRoundOperator()
+{
+}
diff --git a/Engine/Plugins/Experimental/NeuralNetworkInference/Source/NeuralNetworkInference/Private/NeuralOperators/SigmoidOperator.cpp b/Engine/Plugins/Experimental/NeuralNetworkInference/Source/NeuralNetworkInference/Private/NeuralOperators/SigmoidOperator.cpp
new file mode 100644
--- /dev/null
+++ b/Engine/Plugins/Experimental/NeuralNetworkInference/Source/NeuralNetworkInference/Private/NeuralOperators/SigmoidOperator.cpp
@@ -0,0 +1,17 @@
+// Copyright Epic Games, Inc. All Rights Reserved.
+
+#include "NeuralOperators/SigmoidOperator.h"
+
+
+
+/* FSigmoidOperator structors
+ *****************************************************************************/
+
+FSigmoidOperator::FSigmoidOperator(const bool bIsInlinedTensor)
+	: IElementWiseOperator(TEXT("Sigmoid"), 13, EElementWiseOperator::Sigmoid, bIsInlinedTensor)
+{
+}
+
+FSigmoidOperator::~FSigmoidOperator()
+{
+}
diff --git a/Engine/Plugins/Experimental/NeuralNetworkInference/Source/NeuralNetworkInference/Private/NeuralOperators/SignOperator.cpp b/Engine/Plugins/Experimental/NeuralNetworkInference/Source/NeuralNetworkInference/Private/NeuralOperators/SignOperator.cpp
new file mode 100644
--- /dev/null
+++ b/Engine/Plugins/Experimental/NeuralNetworkInference/Source/NeuralNetworkInference/Private/NeuralOperators/SignOperator.cpp
@@ -0,0 +1,17 @@
+// Copyright Epic Games, Inc. All Rights Reserved.
+
+#include "NeuralOperators/SignOperator.h"
+
+
+
+/* FSignOperator structors
+ *****************************************************************************/
+
+FSignOperator::FSignOperator(const bool bIsInlinedTensor)
+	: IElementWiseOperator(TEXT("Sign"), 13, EElementWiseOperator::Sign, bIsInlinedTensor)
+{
+}
+
+FSignOperator::~FSignOperator()
+{
+}
diff --git a/Engine/Plugins/Experimental/NeuralNetworkInference/Source/NeuralNetworkInference/Private/NeuralOperators/SinOperator.cpp b/Engine/Plugins/Experimental/NeuralNetworkInference/Source/NeuralNetworkInference/Private/NeuralOperators/SinOperator.cpp
new file mode 100644
--- /dev/null
+++ b/Engine/Plugins/Experimental/NeuralNetworkInference/Source/NeuralNetworkInference/Private/NeuralOperators/SinOperator.cpp
@@ -0,0 +1,17 @@
+// Copyright Epic Games, Inc. All Rights Reserved.
+
+#include "NeuralOperators/SinOperator.h"
+
+
+
+/* FSinOperator structors
+ *****************************************************************************/
+
+FSinOperator::FSinOperator(const bool bIsInlinedTensor)
+	: IElementWiseOperator(TEXT("Sin"), 7, EElementWiseOperator::Sin, bIsInlinedTensor)
+{
+}
+
+FSinOperator::~FSinOperator()
+{
+}
diff --git a/Engine/Plugins/Experimental/NeuralNetworkInference/Source/NeuralNetworkInference/Private/NeuralOperators/SinhOperator.cpp b/Engine/Plugins/Experimental/NeuralNetworkInference/Source/NeuralNetworkInference/Private/NeuralOperators/SinhOperator.cpp
new file mode 100644
--- /dev/null
+++ b/Engine/Plugins/Experimental/NeuralNetworkInference/Source/NeuralNetworkInference/Private/NeuralOperators/SinhOperator.cpp
@@ -0,0 +1,17 @@
+// Copyright Epic Games, Inc. All Rights Reserved.
+
+#include "NeuralOperators/SinhOperator.h"
+
+
+
+/* FSinhOperator structors
+ *****************************************************************************/
+
+FSinhOperator::FSinhOperator(const bool bIsInlinedTensor)
+	: IElementWiseOperator(TEXT("Sinh"), 9, EElementWiseOperator::Sinh, bIsInlinedTensor)
+{
+}
+
+FSinhOperator::~FSinhOperator()
+{
+}
diff --git a/Engine/Plugins/Experimental/NeuralNetworkInference/Source/NeuralNetworkInference/Private/NeuralOperators/SqrtOperator.cpp b/Engine/Plugins/Experimental/NeuralNetworkInference/Source/NeuralNetworkInference/Private/NeuralOperators/SqrtOperator.cpp
new file mode 100644
--- /dev/null
+++ b/Engine/Plugins/Experimental/NeuralNetworkInference/Source/NeuralNetworkInference/Private/NeuralOperators/SqrtOperator.cpp
@@ -0,0 +1,17 @@
+// Copyright Epic Games, Inc. All Rights Reserved.
+
+#include "NeuralOperators/SqrtOperator.h"
+
+
+
+/* FSqrtOperator structors
+ *****************************************************************************/
+
+FSqrtOperator::FSqrtOperator(const bool bIsInlinedTensor)
+	: IElementWiseOperator(TEXT("Sqrt"), 13, EElementWiseOperator::Sqrt, bIsInlinedTensor)
+{
+}
+
+FSqrtOperator::~FSqrtOperator()
+{
+}
diff --git a/Engine/Plugins/Experimental/NeuralNetworkInference/Source/NeuralNetworkInference/Private/NeuralOperators/SqueezeOperator.cpp b/Engine/Plugins/Experimental/NeuralNetworkInference/Source/NeuralNetworkInference/Private/NeuralOperators/SqueezeOperator.cpp
new file mode 100644
--- /dev/null
+++ b/Engine/Plugins/Experimental/NeuralNetworkInference/Source/NeuralNetworkInference/Private/NeuralOperators/SqueezeOperator.cpp
@@ -0,0 +1,243 @@
+// Copyright Epic Games, Inc. All Rights Reserved.
+
+#include "NeuralOperators/SqueezeOperator.h"
+#include "NeuralNetworkInferenceShaders/CopyCS.h"
+#include "NeuralNetworkInferenceUtils.h"
+#include "NeuralNetworkInferenceUtilsGPU.h"
+
+
+
+/* FSqueezeOperator structors
+ *****************************************************************************/
+
+FSqueezeOperator::FSqueezeOperator(const bool bIsInlinedTensor)
+	: FNeuralOperator(TEXT("Squeeze"), 13, (bIsInlinedTensor ? 0 : -1))
+{
+}
+
+FSqueezeOperator::~FSqueezeOperator()
+{
+}
+
+
+
+/* FSqueezeOperator public functions
+ *****************************************************************************/
+
+bool FSqueezeOperator::ConfigureOutputAndInternalVariablesAndSanityChecks()
+{
+	bIsLoaded = false;
+	// Output initialization and size sanity check
+	if (InlinedTensor < 0)
+	{
+		FNeuralTensor& OutputTensor = GetOutputTensorNoConst();
+		const TArray<FNeuralTensor*>& InputTensors = GetInputTensorsConst();
+		const FNeuralTensor& InputTensor = *InputTensors[0];
+
+		// Initialize OutputNeuralTensor
+		if (InputTensors.Num() > 1)
+		{
+			const FNeuralTensor& AxesTensor = *InputTensors[1];
+			TArray<int64> OutputShape = FSqueezeOperator::Remove1sFromShape(InputTensor, AxesTensor);
+			OutputTensor.SetNumUninitialized(OutputShape, InputTensor.GetDataType());
+		}
+		else
+		{
+			TArray<int64> OutputShape = FSqueezeOperator::Remove1sFromShape(InputTensor);
+			OutputTensor.SetNumUninitialized(OutputShape, InputTensor.GetDataType());
+		}
+		
+	}
+	// Return true
+	bIsLoaded = true;
+	return bIsLoaded;
+}
+
+bool FSqueezeOperator::SetWhetherInputTensorsNeedTransferToGPUForGPUMode()
+{
+	// Enable GPU mode only on the first input tensor
+	const TArray<FNeuralTensor*>& InputTensors = GetInputTensorsConst();
+	InputTensors[0]->SetEnableGPU(true);
+	return true;
+}
+
+void FSqueezeOperator::ForwardCPU()
+{
+	// Sanity check
+	if (!bIsLoaded)
+	{
+		UE_LOG(LogNeuralNetworkInference, Warning, TEXT("FSqueezeOperator::ForwardCPU(): bIsLoaded was false."));
+		return;
+	}
+	// Squeeze is 0-cost if inlined, otherwise, just copy memory from input to output
+	const TArray<FNeuralTensor*> InputTensors = GetInputTensorsConst();
+	const FNeuralTensor& InputTensor = *InputTensors[0];
+	FNeuralTensor& OutputTensor = (InlinedTensor < 0 ? GetOutputTensorNoConst() : *GetInputTensorsNoConst()[0]);
+	// Update InputSizeIfInlined (if inlined)
+	if (InlinedTensor > -1)
+	{
+		InputSizeIfInlined = InputTensor.GetSizes();
+	}
+	// Copy memory (if not inlined)
+	else
+	{
+		OutputTensor.SetFromUnderlyingUInt8ArrayCopy(InputTensor.GetUnderlyingUInt8ArrayRef());
+	}
+	TArray<int64> OutputShape;
+	if (InputTensors.Num() > 1)
+	{
+		const FNeuralTensor& AxesTensor = *InputTensors[1];
+		// Squeeze output
+		OutputShape = FSqueezeOperator::Remove1sFromShape(InputTensor, AxesTensor);
+	}
+	else 
+	{
+		// Squeeze output
+		OutputShape = FSqueezeOperator::Remove1sFromShape(InputTensor);
+	}
+	OutputTensor.ReshapeMove(OutputShape);
+}
+
+void FSqueezeOperator::ForwardGPU_RenderThread(FRDGBuilder* InOutGraphBuilder)
+{
+	// Sanity checks
+	if (!FNeuralNetworkInferenceUtilsGPU::GPUSanityChecks(InOutGraphBuilder, bIsLoaded))
+	{
+		UE_LOG(LogNeuralNetworkInference, Warning, TEXT("FSqueezeOperator::ForwardGPU_RenderThread(): Sanity checks failed."));
+		return;
+	}
+	// Squeeze is 0-cost if inlined, otherwise, just copy memory from input to output
+	const TArray<FNeuralTensor*> InputTensors = GetInputTensorsConst();
+	const FNeuralTensor& InputTensor = *InputTensors[0];
+	FNeuralTensor& OutputTensor = (InlinedTensor < 0 ? GetOutputTensorNoConst() : *GetInputTensorsNoConst()[0]);
+	// If input is empty, it means it has a GetSizes() similar to [a, b, 0, c, ...]
+	if (InputTensor.IsEmpty())
+	{
+		return;
+	}
+	// Update InputSizeIfInlined (if inlined)
+	if (InlinedTensor > -1)
+	{
+		InputSizeIfInlined = InputTensor.GetSizes();
+	}
+	// Copy memory (if not inlined)
+	else
+	{
+		// Set parameters
+		FCopyCS::FParameters* Parameters = InOutGraphBuilder->AllocParameters<FCopyCS::FParameters>();
+		Parameters->Num = OutputTensor.Num();
+		// SRV/UAV variables
+		Parameters->InputSRV = InputTensor.GetBufferSRVRef();
+		Parameters->OutputUAV = OutputTensor.GetBufferUAVRef();
+		// Set shader
+		TShaderMapRef<FCopyCS> ComputeShader(GetGlobalShaderMap(GMaxRHIFeatureLevel));
+		// Run shader
+		const uint32 ThreadGroupCountValueX = FMath::DivideAndRoundUp(Parameters->Num, FCopyCS::THREADGROUP_SIZE_X);
+		FComputeShaderUtils::AddPass(
+			*InOutGraphBuilder,
+			RDG_EVENT_NAME("FSqueezeeOperator-CopyCS()"),
+			ComputeShader,
+			Parameters,
+			FIntVector(ThreadGroupCountValueX, 1, 1));
+	}
+	// Reshape output
+	TArray<int64> OutputShape;
+	if (InputTensors.Num() > 1)
+	{
+		const FNeuralTensor& AxesTensor = *InputTensors[1];
+		// Squeeze output
+		OutputShape = FSqueezeOperator::Remove1sFromShape(InputTensor, AxesTensor);
+	}
+	else
+	{
+		// Squeeze output
+		OutputShape = FSqueezeOperator::Remove1sFromShape(InputTensor);
+	}
+	OutputTensor.ReshapeMove(OutputShape);
+}
+
+void FSqueezeOperator::PostForwardCPU()
+{
+	if (NeedsPostForward())
+	{
+		FNeuralTensor& InputTensor = *GetInputTensorsNoConst()[0];
+		InputTensor.Reshape(InputSizeIfInlined);
+	}
+}
+
+TArray<int64> FSqueezeOperator::Remove1sFromShape(const FNeuralTensor& InTensor, const FNeuralTensor& InAxesTensor) const
+{
+	const TArray<int64>& CurrentShape = InTensor.GetSizes();
+	TArray<int64> OutputShape = CurrentShape;
+	TArray<int64> AxesList = InAxesTensor.GetArrayCopy<int64>();
+
+	// Convert negative indices to the equivalent positive indices
+	for (int64& Axes : AxesList)
+	{
+		if (Axes < 0)
+		{
+			Axes += CurrentShape.Num();
+		}
+	}
+
+	// Sort Axes in descending order
+	AxesList.Sort([](const int64& A, const int64& B) {
+		return A > B;
+		});
+
+	// Remove all duplicates
+	bool isDuplicateRemoved = false;
+	for (int64 CurrentAxesIndex = 0; CurrentAxesIndex < AxesList.Num() - 1; ++CurrentAxesIndex)
+	{
+		if (AxesList[CurrentAxesIndex] == AxesList[CurrentAxesIndex + 1])
+		{
+			AxesList.RemoveAt(CurrentAxesIndex);
+			if (!isDuplicateRemoved)
+			{
+				isDuplicateRemoved = true;
+			}
+		}
+	}
+
+	if (isDuplicateRemoved)
+	{
+		UE_LOG(LogNeuralNetworkInference, Warning, TEXT("FSqueezeOperator::Remove1sFromShape(): Axes input had duplicate(s) and was removed."));
+	}
+
+	for (int64 Axes : AxesList)
+	{
+
+		// 1 = remove from shape
+		if (CurrentShape[Axes] == 1)
+		{
+			OutputShape.RemoveAt(Axes);
+		}
+		else 
+		{
+			UE_LOG(LogNeuralNetworkInference, Warning, TEXT("FSqueezeOperator::Remove1sFromShape(): Axes has an invalid index to remove (Shape value is not equal to one)."));
+			TArray<int64> EmptyTArray;
+			return EmptyTArray;
+		}
+	}
+
+	// Return cleaned OutputShape (without 1s)
+	return OutputShape;
+}
+
+TArray<int64> FSqueezeOperator::Remove1sFromShape(const FNeuralTensor& InTensor) const
+{
+	const TArray<int64>& CurrentShape = InTensor.GetSizes();
+	TArray<int64> OutputShape = CurrentShape;
+
+	for (int64 CurrentIndex = CurrentShape.Num() - 1; CurrentIndex > -1; --CurrentIndex)
+	{
+		// 1 = remove from shape
+		if (CurrentShape[CurrentIndex] == 1)
+		{
+			OutputShape.RemoveAt(CurrentIndex);
+		}
+	}
+
+	// Return cleaned OutputShape (without 1s)
+	return OutputShape;
+}
\ No newline at end of file
diff --git a/Engine/Plugins/Experimental/NeuralNetworkInference/Source/NeuralNetworkInference/Private/NeuralOperators/SubOperator.cpp b/Engine/Plugins/Experimental/NeuralNetworkInference/Source/NeuralNetworkInference/Private/NeuralOperators/SubOperator.cpp
new file mode 100644
--- /dev/null
+++ b/Engine/Plugins/Experimental/NeuralNetworkInference/Source/NeuralNetworkInference/Private/NeuralOperators/SubOperator.cpp
@@ -0,0 +1,17 @@
+// Copyright Epic Games, Inc. All Rights Reserved.
+
+#include "NeuralOperators/SubOperator.h"
+
+
+
+/* FSubOperator structors
+ *****************************************************************************/
+
+FSubOperator::FSubOperator(const TSet<uint32>& InPotentialInlinedTensors)
+	: IMultidirectionalBroadcastOperator(TEXT("Sub"), 13, EMultidirectionalBroadcastOperator::Sub, InPotentialInlinedTensors)
+{
+}
+
+FSubOperator::~FSubOperator()
+{
+}
diff --git a/Engine/Plugins/Experimental/NeuralNetworkInference/Source/NeuralNetworkInference/Private/NeuralOperators/TanOperator.cpp b/Engine/Plugins/Experimental/NeuralNetworkInference/Source/NeuralNetworkInference/Private/NeuralOperators/TanOperator.cpp
new file mode 100644
--- /dev/null
+++ b/Engine/Plugins/Experimental/NeuralNetworkInference/Source/NeuralNetworkInference/Private/NeuralOperators/TanOperator.cpp
@@ -0,0 +1,17 @@
+// Copyright Epic Games, Inc. All Rights Reserved.
+
+#include "NeuralOperators/TanOperator.h"
+
+
+
+/* FTanOperator structors
+ *****************************************************************************/
+
+FTanOperator::FTanOperator(const bool bIsInlinedTensor)
+	: IElementWiseOperator(TEXT("Tan"), 7, EElementWiseOperator::Tan, bIsInlinedTensor)
+{
+}
+
+FTanOperator::~FTanOperator()
+{
+}
diff --git a/Engine/Plugins/Experimental/NeuralNetworkInference/Source/NeuralNetworkInference/Private/NeuralOperators/TanhOperator.cpp b/Engine/Plugins/Experimental/NeuralNetworkInference/Source/NeuralNetworkInference/Private/NeuralOperators/TanhOperator.cpp
new file mode 100644
--- /dev/null
+++ b/Engine/Plugins/Experimental/NeuralNetworkInference/Source/NeuralNetworkInference/Private/NeuralOperators/TanhOperator.cpp
@@ -0,0 +1,28 @@
+// Copyright Epic Games, Inc. All Rights Reserved.
+
+#include "NeuralOperators/TanhOperator.h"
+#include <cmath>
+
+
+
+/* FTanhOperator structors
+ *****************************************************************************/
+
+FTanhOperator::FTanhOperator(const bool bIsInlinedTensor)
+	: IElementWiseOperator(TEXT("Tanh"), 13, EElementWiseOperator::Tanh, bIsInlinedTensor)
+{
+}
+
+FTanhOperator::~FTanhOperator()
+{
+}
+
+
+
+/* FTanhOperator public functions
+ *****************************************************************************/
+
+void FTanhOperator::ForwardCPU()
+{
+	ForwardCPUWithFunction([](const float InValue) { return std::tanh(InValue); });
+}
