NeuralNetworkInference (NNI) plugin:

diff --git a/Engine/Plugins/Experimental/NeuralNetworkInference/Source/NeuralNetworkInference/Private/NeuralTensor.cpp b/Engine/Plugins/Experimental/NeuralNetworkInference/Source/NeuralNetworkInference/Private/NeuralTensor.cpp
new file mode 100644
--- /dev/null
+++ b/Engine/Plugins/Experimental/NeuralNetworkInference/Source/NeuralNetworkInference/Private/NeuralTensor.cpp
@@ -0,0 +1,787 @@
+// Copyright Epic Games, Inc. All Rights Reserved.
+
+#include "NeuralTensor.h"
+#include "Algo/Accumulate.h"
+#include "Containers/ResourceArray.h"
+#include "ModelProto.h"
+#include "NeuralNetworkInferenceUtils.h"
+#include "NeuralNetworkInferenceUtilsGPU.h"
+#include "NeuralTensorResourceArray.h"
+#include "RenderGraphBuilder.h"
+#include "RenderGraphResources.h"
+#include "RenderGraphUtils.h"
+#include "RenderingThread.h"
+#include "Templates/RefCounting.h"
+
+
+
+/* FPrivateNeuralTensor auxiliary functions
+ *****************************************************************************/
+
+class FPrivateNeuralTensor
+{
+public:
+	template <typename T>
+	static FString SanitizeFloat(const T InValue);
+	static void ArrayToSanitizedString(FString& InOutTensorString, const int64 InIndexStart, const int64 InIndexEnd, const int64 InOffset, const ENeuralDataType InDataType, const FNeuralTensor& InTensor);
+
+	static void NDTensorIndexesPlus1(TArray<int32>& InOutImageAreaIndexes, const TArray<int32>& InSizes);
+
+	template <typename T>
+	struct TMultiplies
+	{
+		constexpr T operator()(const T& InLeft, const T& InRight) const { return InLeft * InRight; }
+	};
+};
+
+template <typename T>
+FString FPrivateNeuralTensor::SanitizeFloat(const T InValue)
+{
+	return (FMath::IsFinite(InValue) ? FString::SanitizeFloat(InValue) : FString(TEXT("NaNInf")));
+}
+
+#define FOR_LOOP_FLOAT_TYPE_TO_SANITIZED_STRING(TensorString, InTensor, IndexStart, IndexEnd, Offset, FNumberToStringFunction, DataType) \
+	for (int64 Index = IndexStart; Index < IndexEnd; ++Index) \
+	{ \
+		TensorString += FNumberToStringFunction(InTensor.At<DataType>(Offset + Index)) + TEXT(" "); \
+	}
+
+void FPrivateNeuralTensor::ArrayToSanitizedString(FString& InOutTensorString, const int64 InIndexStart, const int64 InIndexEnd, const int64 InOffset, const ENeuralDataType InDataType, const FNeuralTensor& InTensor)
+{
+	if (InDataType == ENeuralDataType::Float)
+	{
+		FOR_LOOP_FLOAT_TYPE_TO_SANITIZED_STRING(InOutTensorString, InTensor, InIndexStart, InIndexEnd, InOffset, FPrivateNeuralTensor::SanitizeFloat, float);
+	}
+	else if (InDataType == ENeuralDataType::Int32)
+	{
+		FOR_LOOP_FLOAT_TYPE_TO_SANITIZED_STRING(InOutTensorString, InTensor, InIndexStart, InIndexEnd, InOffset, FString::FromInt, int32);
+	}
+	else if (InDataType == ENeuralDataType::Int64)
+	{
+		FOR_LOOP_FLOAT_TYPE_TO_SANITIZED_STRING(InOutTensorString, InTensor, InIndexStart, InIndexEnd, InOffset, FString::FromInt, int64);
+	}
+	else if (InDataType == ENeuralDataType::UInt32)
+	{
+		FOR_LOOP_FLOAT_TYPE_TO_SANITIZED_STRING(InOutTensorString, InTensor, InIndexStart, InIndexEnd, InOffset, FString::FromInt, uint32);
+	}
+	else if (InDataType == ENeuralDataType::UInt64)
+	{
+		FOR_LOOP_FLOAT_TYPE_TO_SANITIZED_STRING(InOutTensorString, InTensor, InIndexStart, InIndexEnd, InOffset, FString::FromInt, uint64);
+	}
+	else
+	{
+		UE_LOG(LogNeuralNetworkInference, Warning, TEXT("FNeuralTensor::ArrayToSanitizedString(): Unknown InDataType = %d used."), (int32)InDataType);
+	}
+}
+
+void FPrivateNeuralTensor::NDTensorIndexesPlus1(TArray<int32>& InOutImageAreaIndexes, const TArray<int32>& InSizes)
+{
+	int64 ImageAreaIndexesIndex = InSizes.Num() - 1;
+	while (ImageAreaIndexesIndex > -1)
+	{
+		++InOutImageAreaIndexes[ImageAreaIndexesIndex];
+		if (InOutImageAreaIndexes[ImageAreaIndexesIndex] == InSizes[ImageAreaIndexesIndex])
+		{
+			InOutImageAreaIndexes[ImageAreaIndexesIndex] = 0;
+			--ImageAreaIndexesIndex;
+		}
+		else
+		{
+			break;
+		}
+	}
+}
+
+
+
+/* FNeuralTensor structors
+ *****************************************************************************/
+
+FNeuralTensor::FNeuralTensor(const ENeuralDataType InDataType, const int64 InVolume, const FString& InName, const ENeuralTensorTypeGPU InTensorTypeGPU)
+	: FNeuralTensor(InDataType, InVolume > 0 ? TArray<int64>({ InVolume }) : TArray<int64>({}), InName, InTensorTypeGPU)
+{
+}
+
+FNeuralTensor::FNeuralTensor(const ENeuralDataType InDataType, const TArray<int64>& InSizes, const FString& InName, const ENeuralTensorTypeGPU InTensorTypeGPU)
+	: Name(InName)
+	, DataType(ENeuralDataType::None)
+	, TensorTypeGPU(InTensorTypeGPU)
+	, bEnableGPU(false)
+{
+	// Memory allocation
+	SetNumUninitialized(InSizes, InDataType);
+}
+
+FNeuralTensor::FNeuralTensor(const ENeuralDataType InDataType, const void* const InValues, const int64 InSizeOfT, const int64 InValueNum, const TArray<int64>& InSizes, const FString& InName, const ENeuralTensorTypeGPU InTensorTypeGPU)
+	: FNeuralTensor(InDataType, InSizes, InName, InTensorTypeGPU)
+{
+	// Sanity check
+	if (IsEmpty())
+	{
+		UE_LOG(LogNeuralNetworkInference, Warning, TEXT("FNeuralTensor(): GetVolume() = 0. Skipping array copy."));
+	}
+	// Memory filling
+	else
+	{
+		SetFromPointer(InValues, InSizeOfT, InValueNum);
+	}
+}
+
+FNeuralTensor::FNeuralTensor(const FNeuralTensor& InTensor)
+	: FNeuralTensor(InTensor.DataType, InTensor.Sizes, InTensor.Name)
+{
+	SetTensorTypeGPU(InTensor.TensorTypeGPU);
+	ArrayCPU = InTensor.ArrayCPU;
+	bEnableGPU = InTensor.bEnableGPU;
+	PooledBuffer = InTensor.PooledBuffer;
+	BufferSRVRef = InTensor.BufferSRVRef;
+	BufferUAVRef = InTensor.BufferUAVRef;
+}
+
+FNeuralTensor& FNeuralTensor::operator=(const FNeuralTensor& InTensor)
+{
+	Name = InTensor.Name;
+	SetNumUninitialized(InTensor.Sizes, InTensor.DataType);
+	SetTensorTypeGPU(InTensor.TensorTypeGPU);
+	ArrayCPU = InTensor.ArrayCPU;
+	bEnableGPU = InTensor.bEnableGPU;
+	PooledBuffer = InTensor.PooledBuffer;
+	BufferSRVRef = InTensor.BufferSRVRef;
+	BufferUAVRef = InTensor.BufferUAVRef;
+
+	return *this;
+}
+
+FNeuralTensor::FNeuralTensor(FNeuralTensor&& InTensor)
+	: FNeuralTensor(InTensor.DataType, InTensor.Sizes, TEXT(""))
+{
+	Swap(Name, InTensor.Name);
+	SetTensorTypeGPU(InTensor.TensorTypeGPU);
+	Swap(ArrayCPU, InTensor.ArrayCPU);
+	bEnableGPU = InTensor.bEnableGPU;
+	Swap(PooledBuffer, InTensor.PooledBuffer);
+	Swap(BufferSRVRef, InTensor.BufferSRVRef);
+	Swap(BufferUAVRef, InTensor.BufferUAVRef);
+}
+
+FNeuralTensor& FNeuralTensor::operator=(FNeuralTensor&& InTensor)
+{
+	Swap(Name, InTensor.Name);
+	SetNumUninitialized(InTensor.Sizes, InTensor.DataType);
+	SetTensorTypeGPU(InTensor.TensorTypeGPU);
+	Swap(ArrayCPU, InTensor.ArrayCPU);
+	bEnableGPU = InTensor.bEnableGPU;
+	Swap(PooledBuffer, InTensor.PooledBuffer);
+	Swap(BufferSRVRef, InTensor.BufferSRVRef);
+	Swap(BufferUAVRef, InTensor.BufferUAVRef);
+
+	return *this;
+}
+
+
+
+/* FNeuralTensor public operators
+ *****************************************************************************/
+
+bool FNeuralTensor::operator==(const FNeuralTensor& InTensorToCopy) const
+{
+	// Dimensions, sizes, and scalar type match --> Check if data does
+	if (Num() == InTensorToCopy.Num() && GetSizes() == InTensorToCopy.GetSizes())
+	{
+		return ArrayCPU == InTensorToCopy.ArrayCPU;
+	}
+	// Dimensions or sizes or scalar type do not match
+	return false;
+}
+
+
+
+/* FNeuralTensor public functions
+ *****************************************************************************/
+
+int64 FNeuralTensor::GetSize(const int32 InDimension) const
+{
+	return (InDimension < GetNumberDimensions() ? Sizes[InDimension] : 1u);
+}
+
+void FNeuralTensor::SetTensorTypeGPU(const ENeuralTensorTypeGPU InTensorTypeGPU)
+{
+	// Sanity check
+	if (PooledBuffer.IsValid() || BufferSRVRef.IsValid() || BufferUAVRef.IsValid())
+	{
+		UE_LOG(LogNeuralNetworkInference, Warning, TEXT("FNeuralTensor-%s::SetTensorTypeGPU(): TensorTypeGPU cannot be modified from %d to %d because the GPU memory has been already initialized."
+			" Modify the GPU type before allocating the GPU memory (e.g., on the FNeuralTensor constructor)."), *Name, TensorTypeGPU, InTensorTypeGPU);
+		return;
+	}
+	// Update TensorTypeGPU
+	TensorTypeGPU = InTensorTypeGPU;
+}
+
+void FNeuralTensor::ToGPU_RenderThread(FRDGBuilder* InOutGraphBuilder)
+{
+// @todo: Volatile or not adding BUF_ShaderResource causes this error:
+//Fatal error: [File:D:/P4/private_dh_research_pitt/Engine/Source/Runtime/Windows/D3D11RHI/Private/D3D11Util.cpp] [Line: 258] 
+//Result failed 
+// at D:/P4/private_dh_research_pitt/Engine/Source/Runtime/Windows/D3D11RHI/Private/D3D11VertexBuffer.cpp:109 
+// with error E_INVALIDARG
+	// Idea:
+	// - BUF_Volatile: Updated multiple times in a frame, but does not imply a lifetime of 1 frame. E.g. a vertex buffer you update every frame with new vertices.
+	// - BUF_Transient: Used during 1 frame. Volatile and transient are not mutually exclusive.
+	// - BUF_KeepCPUAccessible: Not needed, I can just copy the final GPU memory back to RAM at the very end
+	// Call ToGPU_RenderThread with the right flags
+	if (TensorTypeGPU == ENeuralTensorTypeGPU::Generic)
+	{
+		return ToGPU_RenderThread(InOutGraphBuilder, BUF_ShaderResource | BUF_UnorderedAccess, true);
+	}
+	else if (TensorTypeGPU == ENeuralTensorTypeGPU::Input)
+	{
+		return ToGPU_RenderThread(InOutGraphBuilder, BUF_ShaderResource | BUF_UnorderedAccess, true);
+	}
+	else if (TensorTypeGPU == ENeuralTensorTypeGPU::IntermediateNotInitialized)
+	{
+		return ToGPU_RenderThread(InOutGraphBuilder, BUF_ShaderResource | BUF_UnorderedAccess | BUF_Transient, false);
+	}
+	else if (TensorTypeGPU == ENeuralTensorTypeGPU::IntermediateInitialized)
+	{
+		return ToGPU_RenderThread(InOutGraphBuilder, BUF_ShaderResource | BUF_UnorderedAccess, true);
+	}
+	else if (TensorTypeGPU == ENeuralTensorTypeGPU::Output)
+	{
+		return ToGPU_RenderThread(InOutGraphBuilder, BUF_ShaderResource | BUF_UnorderedAccess, false);
+	}
+	else if (TensorTypeGPU == ENeuralTensorTypeGPU::Weight)
+	{
+		return ToGPU_RenderThread(InOutGraphBuilder, BUF_ShaderResource | BUF_Static, true);
+	}
+	UE_LOG(LogNeuralNetworkInference, Warning, TEXT("FNeuralTensor-%s::ToGPU_RenderThread(): Unimplemented TensorTypeGPU = %d. Assuming ENeuralTensorTypeGPU::Generic."), *Name, (int32)TensorTypeGPU);
+	return ToGPU_RenderThread(InOutGraphBuilder, BUF_UnorderedAccess, true);
+}
+
+void FNeuralTensor::ToGPU_RenderThread(FRDGBuilder* InOutGraphBuilder, const EBufferUsageFlags InEBufferUsageFlags, const bool bInShouldCopyFromCPU)
+{
+	// Sanity checks
+	if (!bEnableGPU || IsEmpty())
+	{
+		return;
+	}
+	checkf(!Name.IsEmpty(), TEXT("FNeuralTensor::ToGPU_RenderThread(): Name cannot be empty."));
+	checkf(IsInRenderingThread(), TEXT("FNeuralTensor-%s::ToGPU_RenderThread(): IsInRenderingThread() must be true."), *Name);
+	checkf(InOutGraphBuilder, TEXT("FNeuralTensor-%s::ToGPU_RenderThread(): InOutGraphBuilder is nullptr."), *Name);
+	// Not SRV-only and not UAV/SRV
+	checkf(EnumHasAnyFlags(InEBufferUsageFlags, BUF_ShaderResource|BUF_UnorderedAccess),
+		TEXT("FNeuralTensor-%s::ToGPU_RenderThread(): Unexpected case InEBufferUsageFlags = %d."), *Name, (uint32)InEBufferUsageFlags);
+	// If SRV-only and bInShouldCopyFromCPU == false
+	if (!EnumHasAnyFlags(InEBufferUsageFlags, BUF_UnorderedAccess) && !bInShouldCopyFromCPU)
+	{
+		UE_LOG(LogNeuralNetworkInference, Warning, TEXT("FNeuralTensor-%s::ToGPU_RenderThread(): bInShouldCopyFromCPU must be true for SRVs (because they cannot be edited). Assumed true."), *Name);
+	}
+	// Create BufferRef
+	FRDGBufferDesc BufferDesc;
+	BufferDesc.BytesPerElement = FDataType::GetSize(DataType);
+	BufferDesc.NumElements = Num();
+	BufferDesc.Usage = InEBufferUsageFlags;
+	BufferDesc.UnderlyingType = FRDGBufferDesc::EUnderlyingType::VertexBuffer;
+	FRDGBufferRef BufferRef = (bInShouldCopyFromCPU ? CreateVertexBuffer(*InOutGraphBuilder, *Name, BufferDesc, ArrayCPU.GetData(), NumInBytes(), ERDGInitialDataFlags::NoCopy)
+		: InOutGraphBuilder->CreateBuffer(BufferDesc, *Name));
+	// Recreate BufferSRVRef
+	if (EnumHasAnyFlags(InEBufferUsageFlags, BUF_ShaderResource))
+	{
+		BufferSRVRef = MakeShared<FRDGBufferSRVRef>(InOutGraphBuilder->CreateSRV(BufferRef, FDataType::GetPixelFormat(DataType)));
+	}
+	else
+	{
+		BufferSRVRef.Reset();
+	}
+	// Recreate BufferUAVRef
+	if (EnumHasAnyFlags(InEBufferUsageFlags, BUF_UnorderedAccess))
+	{
+		BufferUAVRef = MakeShared<FRDGBufferUAVRef>(InOutGraphBuilder->CreateUAV(BufferRef, FDataType::GetPixelFormat(DataType)));
+	}
+	else
+	{
+		BufferUAVRef.Reset();
+	}
+	// Recreate PooledBuffer for future runs
+	PooledBuffer = MakeShared<TRefCountPtr<FRDGPooledBuffer>>();
+	*PooledBuffer = InOutGraphBuilder->ConvertToExternalBuffer(BufferRef);
+}
+
+void FNeuralTensor::UpdateSRVAndOrUAV_RenderThread(FRDGBuilder* InOutGraphBuilder)
+{
+	if (!bEnableGPU)
+	{
+		return;
+	}
+	// Sanity checks
+	checkf(IsInRenderingThread(), TEXT("FNeuralTensor-%s::UpdateSRVAndOrUAV_RenderThread(): IsInRenderingThread() must be true."), *Name);
+	checkf(PooledBuffer.IsValid() && InOutGraphBuilder, TEXT("FNeuralTensor-%s::UpdateSRVAndOrUAV_RenderThread(): IPooledBuffer and InOutGraphBuilder cannot be nullptrs."), *Name);
+	// Register BufferRef
+	FRDGBufferRef BufferRef = InOutGraphBuilder->RegisterExternalBuffer(*PooledBuffer);
+	// Recreate BufferSRVRef
+	if (BufferSRVRef.IsValid())
+	{
+		BufferSRVRef = MakeShared<FRDGBufferSRVRef>(InOutGraphBuilder->CreateSRV(BufferRef, FDataType::GetPixelFormat(DataType)));
+	}
+	// Recreate BufferUAVRef
+	if (BufferUAVRef.IsValid())
+	{
+		BufferUAVRef = MakeShared<FRDGBufferUAVRef>(InOutGraphBuilder->CreateUAV(BufferRef, FDataType::GetPixelFormat(DataType)));
+	}
+}
+
+void FNeuralTensor::ToCPU_RenderThread(FRDGBuilder* InOutGraphBuilder)
+{
+	// Sanity checks
+	if (!bEnableGPU || IsEmpty())
+	{
+		return;
+	}
+	if (!FNeuralNetworkInferenceUtilsGPU::GPUSanityChecks(InOutGraphBuilder))
+	{
+		UE_LOG(LogNeuralNetworkInference, Warning, TEXT("FNeuralTensor-%s::ToCPU_RenderThread(): Sanity checks failed."), *Name);
+		return;
+	}
+	// Read GPU memory back into CPU
+	InOutGraphBuilder->AddPass(
+		RDG_EVENT_NAME("FNeuralTensor(%s)::ToCPU()", *Name),
+		ERDGPassFlags::None,
+		[this](FRHICommandListImmediate& RHICmdList)
+	{
+		const int64 VolumeInBytes = NumInBytes();
+		FRHIBuffer* VertexBuffer = (*PooledBuffer)->GetRHI();
+		const void* const BufferData = RHICmdList.LockBuffer(VertexBuffer, 0, VolumeInBytes, RLM_ReadOnly);
+		FMemory::Memcpy((void*)ArrayCPU.GetData(), BufferData, VolumeInBytes);
+		RHICmdList.UnlockBuffer(VertexBuffer);
+	});
+}
+
+TRefCountPtr<class FRDGPooledBuffer>& FNeuralTensor::GetPooledBuffer() const
+{
+	// Sanity checks
+	checkf(bEnableGPU, TEXT("FNeuralTensor-%s::GetPooledBuffer(): bEnableGPU must be true."), *Name);
+	checkf(IsInRenderingThread(), TEXT("FNeuralTensor-%s::GetPooledBuffer(): IsInRenderingThread() must be true."), *Name);
+	checkf(BufferUAVRef.IsValid(), TEXT("FNeuralTensor-%s::GetPooledBuffer(): PooledBuffer cannot be nullptr."), *Name);
+	// Return PooledBuffer
+	return *PooledBuffer;
+}
+
+const FRDGBufferSRVRef FNeuralTensor::GetBufferSRVRef() const
+{
+	// Sanity checks
+	checkf(bEnableGPU, TEXT("FNeuralTensor-%s::GetBufferSRVRef(): bEnableGPU must be true."), *Name);
+	checkf(IsInRenderingThread(), TEXT("FNeuralTensor-%s::GetBufferSRVRef(): IsInRenderingThread() must be true."), *Name);
+	checkf(BufferSRVRef.IsValid(), TEXT("FNeuralTensor-%s::GetBufferSRVRef(): BufferSRVRef was a nullptr, 2 possible causes: 1) FNeuralTensor::ToGPU_RenderThread() was not called. 2) The tensor was empty."), *Name);
+	// Return BufferSRVRef
+	return *BufferSRVRef;
+}
+
+FRDGBufferUAVRef FNeuralTensor::GetBufferUAVRef()
+{
+	// Sanity checks
+	checkf(bEnableGPU, TEXT("FNeuralTensor-%s::GetBufferUAVRef(): bEnableGPU must be true."), *Name);
+	checkf(IsInRenderingThread(), TEXT("FNeuralTensor-%s::GetBufferUAVRef(): IsInRenderingThread() must be true."), *Name);
+	checkf(BufferUAVRef.IsValid(), TEXT("FNeuralTensor-%s::GetBufferUAVRef(): BufferUAVRef cannot be nullptr."), *Name);
+	// Return BufferUAVRef
+	return *BufferUAVRef;
+}
+
+void FNeuralTensor::SetNumUninitialized(const FNeuralTensor& InTensor, const bool bInAllowShrinking)
+{
+	SetNumUninitialized(InTensor.GetSizes(), InTensor.GetDataType(), bInAllowShrinking);
+}
+
+void FNeuralTensor::SetNumUninitialized(const TArray<int64>& InSizes, const ENeuralDataType InDataType, const bool bInAllowShrinking)
+{
+	// Update DataType
+	if (InDataType != ENeuralDataType::None)
+	{
+		DataType = InDataType;
+	}
+	// Update Sizes
+	Sizes = InSizes;
+	// Re-initialize ArrayCPU
+	if (Sizes.Num() > 0)
+	{
+		Volume = FNeuralNetworkInferenceUtils::Product<int64>(Sizes);
+		const int64 VolumeInBytes = Volume * FDataType::GetSize(DataType);
+		if (VolumeInBytes != ArrayCPU.Num())
+		{
+			ArrayCPU.SetNumUninitialized(VolumeInBytes, bInAllowShrinking); // Pre-allocate TArray
+		}
+	}
+	else
+	{
+		Volume = 0;
+		ArrayCPU.Empty();
+	}
+}
+
+void FNeuralTensor::SetFromUnderlyingUInt8ArrayCopy(const TArray<uint8>& InArray)
+{
+	if (NumInBytes() != InArray.Num())
+	{
+		UE_LOG(LogNeuralNetworkInference, Warning,
+			TEXT("FNeuralTensor::SetFromUnderlyingUInt8ArrayCopy(): NumInBytes() == InArray.Num() failed, %d != %d."), NumInBytes(), InArray.Num());
+		return;
+	}
+	ArrayCPU = InArray;
+}
+
+bool FNeuralTensor::SetFromTensorProto(const FTensorProto* const InTensorProto, const FString& InTensorName, const ENeuralTensorTypeGPU InTensorTypeGPU)
+{
+	if (!InTensorProto)
+	{
+		UE_LOG(LogNeuralNetworkInference, Warning, TEXT("FNeuralNetworkFromONNXTranslator::SetFromTensorProto(): InTensorProto was a nullptr."));
+		return false;
+	}
+
+	// const FString& InExternalDataDirectory = TEXT("") - @param InExternalDataDirectory is only required if InTensorProto->ExternalData is being used.
+
+	// Create Tensor
+	Name = InTensorName;
+	TensorTypeGPU = InTensorTypeGPU;
+	// Memory allocation
+	SetNumUninitialized(InTensorProto->Dimensions, InTensorProto->GetDataTypeFromTensorProtoDataType());
+
+	// RawData
+	if (!InTensorProto->RawData.IsEmpty())
+	{
+		SetFromUnderlyingUInt8ArrayCopy(InTensorProto->RawData);
+	}
+
+	// InTensorProto->ExternalData
+	else if (!InTensorProto->ExternalData.IsEmpty())
+	{
+#if WITH_EDITOR
+#if PLATFORM_WINDOWS
+		UE_LOG(LogNeuralNetworkInference, Warning, TEXT("FNeuralNetworkFromONNXTranslator::SetFromTensorProto(): DEPRECATED CODE, OTXT no longer functional."));
+		// // Sanity check
+		// if (InTensorProto->ExternalData[0].Key != TEXT("location") || InTensorProto->ExternalData[0].Value.Len() < 1)
+		// {
+		// 	UE_LOG(LogNeuralNetworkInference, Warning, TEXT("FNeuralNetworkFromONNXTranslator::SetFromTensorProto(): InTensorProto->ExternalData[0].Key = %s != \"location\""
+		// 		" || InTensorProto->ExternalData[0].Value.Len() = %d (should be > 0)."), *InTensorProto->ExternalData[0].Key, InTensorProto->ExternalData[0].Value.Len());
+		// 	return false;
+		// }
+		// // Read neural tensor from binary data
+		// const FString BinaryWeightFilePath = InExternalDataDirectory / InTensorProto->ExternalData[0].Value;
+		// if (!FModelProtoFileReader::ReadWeightsFromOtxtBinaryFile((char*)GetData(), Num() * FDataType::GetSize(NeuralDataType), BinaryWeightFilePath))
+		// {
+		// 	UE_LOG(LogNeuralNetworkInference, Warning, TEXT("FNeuralNetworkFromONNXTranslator::SetFromTensorProto(): Could not read binary file: %s."), *BinaryWeightFilePath);
+		// 	return false;
+		// }
+#else //PLATFORM_WINDOWS
+		UE_LOG(LogNeuralNetworkInference, Warning, TEXT("FNeuralNetworkFromONNXTranslator::SetFromTensorProto(): Only implemented in Windows."));
+		return false;
+#endif //PLATFORM_WINDOWS
+#else //WITH_EDITOR
+		UE_LOG(LogNeuralNetworkInference, Warning, TEXT("FNeuralNetworkFromONNXTranslator::SetFromTensorProto(): InTensorProto->ExternalData should never be used in non-Editor mode."));
+		return false;
+#endif //WITH_EDITOR
+	}
+	else
+	{
+		UE_LOG(LogNeuralNetworkInference, Warning, TEXT("FNeuralNetworkFromONNXTranslator::SetFromTensorProto(): InTensorProto was empty (RawData and ExternalData)."));
+		return false;
+	}
+
+	// No issues --> Read successfully
+	return true;
+}
+
+bool FNeuralTensor::Flip(const int32 InDimension)
+{
+	// Sanity check
+	if (InDimension >= GetNumberDimensions())
+	{
+		UE_LOG(LogNeuralNetworkInference, Warning,
+			TEXT("FNeuralTensor-%s::Transpose(): InDimension < GetNumberDimensions() failed, %d >= %d."), *Name, InDimension, GetNumberDimensions());
+		return false;
+	}
+
+	// Find offset created for all dimensions > InDimension
+	const int32 NumberDimensions = GetNumberDimensions();
+	int64 DimensionOffset = 1;
+	for (int32 DimensionIndex = InDimension + 1; DimensionIndex < NumberDimensions; ++DimensionIndex)
+	{
+		DimensionOffset *= Sizes[DimensionIndex];
+	}
+	const int64 BytesPerIndex = FDataType::GetSize(DataType);
+	const int64 DimensionOffsetInBytes = DimensionOffset * BytesPerIndex;
+
+	// Fill TensorNDIndexes and TensorNDSizes
+	TArray<int32> TensorNDIndexes;
+	TArray<int32> TensorNDSizes;
+	{
+		const int32 TensorNDSize = InDimension + 1; // NumberDimensions
+		TensorNDIndexes.Init(0, TensorNDSize);
+		TensorNDSizes.Init(1, TensorNDSize);
+		for (int32 NDTensorSizeIndex = 0; NDTensorSizeIndex < TensorNDSize; ++NDTensorSizeIndex)
+		{
+			TensorNDSizes[NDTensorSizeIndex] *= Sizes[NDTensorSizeIndex]; // Idea: Sizes=[1, 2, 3, 4], InDimension = 2, then TensorNDSizes = [1, 2x3x4]
+		}
+	}
+
+	// Flip each value
+	TArray<uint8> NewArrayOnCPU;
+	NewArrayOnCPU.SetNumUninitialized(NumInBytes());
+	for (int64 TensorIndex = 0; TensorIndex < Num(); TensorIndex += DimensionOffset)
+	{
+		// Get FlippedTensorIndex
+		int64 FlippedTensorIndex = TensorNDIndexes[0];
+		for (int64 DimensionIndex = 1; DimensionIndex < TensorNDIndexes.Num(); ++DimensionIndex)
+		{
+			//FlippedTensorIndex = (TensorNDIndexes[0] * Sizes[1] + TensorNDIndexes[1]) * Sizes[2] + TensorNDIndexes[2] ...;
+			FlippedTensorIndex = FlippedTensorIndex * Sizes[DimensionIndex] + TensorNDIndexes[DimensionIndex];
+		}
+		// Remove last index (that makes it a normal index) and replace with its flipped equivalent
+		FlippedTensorIndex = FlippedTensorIndex + Sizes[InDimension] - 1 - 2 * TensorNDIndexes.Last();
+		// Flip TensorIndex value
+		FMemory::Memcpy(&NewArrayOnCPU.GetData()[TensorIndex * BytesPerIndex], &ArrayCPU.GetData()[FlippedTensorIndex * DimensionOffsetInBytes], DimensionOffsetInBytes);
+		// Increase TensorNDIndexes
+		FPrivateNeuralTensor::NDTensorIndexesPlus1(TensorNDIndexes, TensorNDSizes);
+	}
+	Swap(NewArrayOnCPU, ArrayCPU);
+	return true;
+}
+
+bool FNeuralTensor::Flip(const int32 InDimensionFirst, const int32 InDimensionLast)
+{
+	// Sanity checks
+	if (InDimensionFirst >= InDimensionLast)
+	{
+		UE_LOG(LogNeuralNetworkInference, Warning, TEXT("FNeuralTensor-%s::Transpose(): InDimensionFirst < InDimensionLast failed, %d >= %d."), *Name, InDimensionFirst, InDimensionLast);
+		return false;
+	}
+	else if (InDimensionLast > GetNumberDimensions())
+	{
+		UE_LOG(LogNeuralNetworkInference, Warning, TEXT("FNeuralTensor-%s::Transpose(): InDimensionLast < GetNumberDimensions() failed, %d >= %d."), *Name, InDimensionLast, GetNumberDimensions());
+		return false;
+	}
+	// Flip
+	for (int32 DimensionIndex = InDimensionFirst; DimensionIndex < InDimensionLast; ++DimensionIndex)
+	{
+		if (!Flip(DimensionIndex))
+		{
+			return false;
+		}
+	}
+	return true;
+}
+
+bool FNeuralTensor::Transpose()
+{
+	const int32 NumberDimensions = GetNumberDimensions();
+	// The transpose of a 0D tensor is itself
+	if (NumberDimensions > 0)
+	{
+		// 1-D tensors
+		if (NumberDimensions == 1)
+		{
+			Sizes.Push(1);
+		}
+		// 2-D tensors
+		else if (NumberDimensions == 2)
+		{
+			// Fill NewArrayOnCPU
+			TArray<uint8> NewArrayOnCPU;
+			NewArrayOnCPU.SetNumUninitialized(NumInBytes());
+			uint8* NewArrayOnCPUPtr = NewArrayOnCPU.GetData();
+			const uint8* const ArrayOnCPUPtr = ArrayCPU.GetData();
+			const int64 Height = Sizes[0];
+			const int64 Width = Sizes[1];
+			const int64 Bytes = FDataType::GetSize(DataType);
+			for (int64 Y = 0; Y < Height; ++Y)
+			{
+				for (int64 X = 0; X < Width; ++X)
+				{
+					for (int64 B = 0; B < Bytes; ++B)
+					{
+						NewArrayOnCPU/*Ptr*/[(X * Height + Y)*Bytes + B] = ArrayCPU/*Ptr*/[(Y * Width + X)*Bytes + B];
+					}
+				}
+			}
+			Swap(NewArrayOnCPU, ArrayCPU);
+		}
+		else
+		{
+			UE_LOG(LogNeuralNetworkInference, Warning, TEXT("FNeuralTensor-%s::Transpose(): Unexpected case NumberDimensions = %d != 1 || 2."), *Name, NumberDimensions);
+			return false;
+		}
+		// Swap W <-> H
+		Swap(Sizes[0], Sizes[1]);
+	}
+	return true;
+}
+
+bool FNeuralTensor::Reshape(const TArray<int64>& InSizes)
+{
+	TArray<int64> NewSizes = InSizes;
+	return ReshapeMove(NewSizes);
+}
+
+bool FNeuralTensor::ReshapeMove(TArray<int64>& InSizes)
+{
+	const int64 NewVolume = Algo::Accumulate(InSizes, int64(1), FPrivateNeuralTensor::TMultiplies<int64>());
+	if (Volume == NewVolume)
+	{
+		Swap(Sizes, InSizes);
+		return true;
+	}
+	UE_LOG(LogNeuralNetworkInference, Warning, TEXT("FNeuralTensor-%s::ReshapeMove(): Volume == NewVolume failed, %d != %d."), *Name, Volume, NewVolume);
+	return false;
+}
+
+FString FNeuralTensor::ToString(const int64 InMaxNumberElementsToDisplay, const bool bInReturnOnlyData) const
+{
+	FString TensorString;
+	if (!bInReturnOnlyData)
+	{
+		TensorString += (Name.Len() > 0 ? Name : FString(TEXT("Unnamed FNeuralTensor"))) + TEXT(": ")
+			// DataType
+			+ FDataType::ToString(DataType) + TEXT(", ")
+			+ (TensorTypeGPU == ENeuralTensorTypeGPU::Generic ? TEXT("Generic")
+				: TensorTypeGPU == ENeuralTensorTypeGPU::Input ? TEXT("Input")
+				: TensorTypeGPU == ENeuralTensorTypeGPU::IntermediateNotInitialized ? TEXT("IntermediateNotInitialized")
+				: TensorTypeGPU == ENeuralTensorTypeGPU::IntermediateInitialized ? TEXT("IntermediateInitialized")
+				: TensorTypeGPU == ENeuralTensorTypeGPU::Output ? TEXT("Output")
+				: TensorTypeGPU == ENeuralTensorTypeGPU::Weight ? TEXT("Weight")
+				: TEXT("Unknown"))
+			// Volume and sizes
+			+ TEXT(", volume=") + FString::FromInt(Num()) + TEXT(", sizes={");
+		// Add sizes
+		if (Sizes.Num() > 0)
+		{
+			for (const int64 Size : Sizes)
+			{
+				TensorString += FString::FromInt(Size) + TEXT(" ");
+			}
+			TensorString[TensorString.Len() - 1] = '}';
+		}
+		else
+		{
+			TensorString += TEXT("}");
+
+		}
+		TensorString += TEXT(", data=[");
+	}
+	else
+	{
+		TensorString += TEXT("[");
+	}
+	// Add tensor data
+	if (Num() > 0)
+	{
+		if (InMaxNumberElementsToDisplay < 1 || Num() <= InMaxNumberElementsToDisplay)
+		{
+			// 1D
+			if (GetNumberDimensions() == 1)
+			{
+				// Eg for sizes{ 1, 2 }: [20 10 9 2]
+				FPrivateNeuralTensor::ArrayToSanitizedString(TensorString, 0, Num(), 0, DataType, *this);
+				TensorString[TensorString.Len() - 1] = ']';
+			}
+			// Eg for sizes {1, 2}: [[20 10] [9 2]]
+			else
+			{
+				// Add initial brackets '['
+				for (int64 Index = 0; Index < GetNumberDimensions() - 1; ++Index)
+				{
+					TensorString += TEXT("[");
+				}
+				// Add text
+				const int64 Stride = Sizes.Last();
+				const int64 NumberRows = Num() / Stride;
+				for (int64 StrideIndex = 0; StrideIndex < NumberRows; ++StrideIndex) //0, 1, 2
+				{
+					FPrivateNeuralTensor::ArrayToSanitizedString(TensorString, 0, Stride, StrideIndex * Stride, DataType, *this);
+					// ']' for last dimension
+					TensorString[TensorString.Len() - 1] = ']';
+					int64 NumberBracketsClosed = 1;
+					// Extra ']' for additional dimensions
+					int64 Value = 1;
+					const int64 NextStrideIndex = StrideIndex + 1;
+					for (int32 RemainderIndex = Sizes.Num() - 2; RemainderIndex > -1; --RemainderIndex)
+					{
+						Value *= Sizes[RemainderIndex];
+						if (NextStrideIndex % Value == 0)
+						{
+							++NumberBracketsClosed;
+							TensorString += TEXT("]");
+						}
+						else
+						{
+							break;
+						}
+					}
+					// Extra '[' for following dimensions (unless we are in the last element)
+					if (NextStrideIndex < NumberRows)
+					{
+						TensorString += TEXT(", ");
+						for (int64 BracketIndex = 0; BracketIndex < NumberBracketsClosed; ++BracketIndex)
+						{
+							TensorString += TEXT("[");
+						}
+					}
+				}
+			}
+		}
+		// Display exactly InMaxNumberElementsToDisplay components
+		else
+		{
+			// Display first InMaxNumberElementsToDisplay/2 components
+			FPrivateNeuralTensor::ArrayToSanitizedString(TensorString, 0, InMaxNumberElementsToDisplay / 2, 0, DataType, *this);
+			TensorString += TEXT("... ");
+			// Display last InMaxNumberElementsToDisplay/2 components
+			FPrivateNeuralTensor::ArrayToSanitizedString(TensorString, Num() - InMaxNumberElementsToDisplay / 2, Num(), 0, DataType, *this);
+			TensorString[TensorString.Len() - 1] = ']';
+		}
+	}
+	else
+	{
+		TensorString += TEXT("]");
+	}
+	return TensorString;
+}
+
+
+
+/* FNeuralTensor private functions
+ *****************************************************************************/
+
+void FNeuralTensor::SetFromPointer(const void* const InData, const int64 InSizeOfT, const int64 InDataSize)
+{
+	// Sanity checks
+	if (Num() != InDataSize || NumInBytes() != InSizeOfT * InDataSize)
+	{
+		UE_LOG(LogNeuralNetworkInference, Warning, TEXT("FNeuralTensor-%s::SetFromPointer: Num() == InDataSize failed, %d vs. %d, or NumInBytes() == sizeof(T) x InDataSize failed, %d vs. %d."
+			" If you want to modify the dimensions of FNeuralTensor, call SetNumUninitialized() first."),
+			*Name, Num(), InDataSize, NumInBytes(), InSizeOfT * InDataSize);
+	}
+	else
+	{
+		// Deep copy
+		FMemory::Memcpy(ArrayCPU.GetData(), InData, NumInBytes());
+	}
+}
+
+bool FNeuralTensor::CheckTAndDataTypeResult(const bool bInCheckTAndDataTypeResult, const int64 InSizeOfT) const
+{
+	const int64 SizeOfDataType = FDataType::GetSize(DataType);
+	if (!bInCheckTAndDataTypeResult)
+	{
+		const FString DataTypeString = FDataType::ToString(DataType);
+		// sizeof(T) and DataType do not match
+		if (SizeOfDataType != InSizeOfT)
+		{
+			UE_LOG(LogNeuralNetworkInference, Warning, TEXT("FNeuralTensor-%s::CheckTAndDataTypeResult() failed: DataType = %s, but sizeof(%s) = %d != sizeof(T) = %d."),
+				*Name, *DataTypeString, *DataTypeString, SizeOfDataType, InSizeOfT);
+		}
+		// sizeof(T) matches, but not the expected DataType
+		else
+		{
+			UE_LOG(LogNeuralNetworkInference, Warning, TEXT("FNeuralTensor-%s::CheckTAndDataTypeResult() failed: DataType = %s, but used a different DataType with the same sizeof(%s) of %d."),
+				*Name, *DataTypeString, *DataTypeString, InSizeOfT);
+		}
+	}
+	return bInCheckTAndDataTypeResult;
+}
diff --git a/Engine/Plugins/Experimental/NeuralNetworkInference/Source/NeuralNetworkInference/Private/NeuralTensorResourceArray.h b/Engine/Plugins/Experimental/NeuralNetworkInference/Source/NeuralNetworkInference/Private/NeuralTensorResourceArray.h
new file mode 100644
--- /dev/null
+++ b/Engine/Plugins/Experimental/NeuralNetworkInference/Source/NeuralNetworkInference/Private/NeuralTensorResourceArray.h
@@ -0,0 +1,27 @@
+// Copyright Epic Games, Inc. All Rights Reserved.
+
+#pragma once
+
+#include "Containers/ResourceArray.h"
+
+/** Resource array to pass  */
+class FNeuralTensorResourceArray : public FResourceArrayInterface
+{
+public:
+	FNeuralTensorResourceArray(void* InData, const uint32 InSize)
+		: Data(InData)
+		, Size(InSize)
+	{
+	}
+
+	FORCEINLINE virtual const void* GetResourceData() const override { return Data; }
+	FORCEINLINE virtual uint32 GetResourceDataSize() const override { return Size; }
+	FORCEINLINE virtual void Discard() override { }
+	FORCEINLINE virtual bool IsStatic() const override { return false; }
+	FORCEINLINE virtual bool GetAllowCPUAccess() const override { return false; }
+	FORCEINLINE virtual void SetAllowCPUAccess(bool bInNeedsCPUAccess) override { }
+
+private:
+	void* Data;
+	uint32 Size;
+};
diff --git a/Engine/Plugins/Experimental/NeuralNetworkInference/Source/NeuralNetworkInference/Public/NeuralTensor.h b/Engine/Plugins/Experimental/NeuralNetworkInference/Source/NeuralNetworkInference/Public/NeuralTensor.h
new file mode 100644
--- /dev/null
+++ b/Engine/Plugins/Experimental/NeuralNetworkInference/Source/NeuralNetworkInference/Public/NeuralTensor.h
@@ -0,0 +1,461 @@
+// Copyright Epic Games, Inc. All Rights Reserved.
+
+#pragma once
+
+#include "CoreMinimal.h"
+#include "NeuralEnumClasses.h"
+#include <algorithm>
+#include "NeuralTensor.generated.h"
+
+class FRDGBuilder;
+using FRDGBufferSRVRef = class FRDGBufferSRV*;
+using FRDGBufferUAVRef = class FRDGBufferUAV*;
+template<typename T> class TRefCountPtr;
+enum class EBufferUsageFlags : uint32;
+struct FTensorProto;
+
+
+/**
+ * Although conceptually this could apply to both the CPU and GPU versions, in practice only the GPU performance is affected by this setting.
+ * Input and Intermediate(Not)Initialized currently share the same attributes because input might become intermediate (e.g., if input tensor fed into a ReLU, which simply modifies
+ * the input FNeuralTensor). However, Intermediate(Not)Initialized and Output do not copy the memory from CPU to GPU but rather simply allocates it.
+ * Output might also become Intermediate(Not)Initialized (e.g., if Output -> ReLU -> Output), so it is kept as ReadWrite rather than written once to account for this.
+ */
+UENUM()
+enum class ENeuralTensorTypeGPU : uint8
+{
+	Generic,					/** Generic tensor that works in every situation (ReadWrite), although it might not be the most efficient one. */
+	Input,						/** Input tensor of the UNeuralNetworkLegacy. Copied from CPU and ReadWrite (but usually ReadOnly). */
+	IntermediateNotInitialized,	/** Intermediate tensor of the UNeuralNetworkLegacy (output of at least a layer and input of at least some other layer). Not copied from CPU, ReadWrite, and transient. */
+	IntermediateInitialized,	/** Intermediate tensor that is initialized with CPU data (e.g., XWithZeros in FConvTranpose). Copied from CPU. */
+	Output,						/** Output tensor of the UNeuralNetworkLegacy. Not copied from CPU and ReadWrite. */
+	Weight						/** Weights of a particular operator/layer. Copied from CPU, ReadOnly, and initialized from CPU memory. */
+};
+
+
+
+/**
+ * This is an auxiliary class. See UNeuralNetworkLegacy for a high-level wrapper of the whole NeuralNetworkInference plugin. The UNeuralNetworkLegacy header
+ * documentation also includes some code examples.
+ *
+ * FNeuralTensor represents a tensor of the UNeuralNetworkLegacy model.
+ * Its CPU operations are very similar to those of TArray<T>. Most of its operations run on CPU, so `ToGPU_RenderThread()` must be called before running on
+ * GPU and after running any FNeuralTensor function that modified the CPU memory.
+ */
+USTRUCT()
+struct NEURALNETWORKINFERENCE_API FNeuralTensor
+{
+	GENERATED_BODY()
+
+public:
+	/**
+	 * It allocates the desired memory.
+	 * @param InVolume Set to 0 if memory allocation is not required or the final size is unknown. Values smaller than 0 will be clipped to 0.
+	 * @param InSizes Set to empty (or omit argument) if memory allocation is not required or the final size is unknown.
+	 * @param InName Used for GPU debugging and the ToString() function.
+	*/
+	explicit FNeuralTensor(const ENeuralDataType InDataType, const int64 InVolume, const FString& InName = TEXT("FNeuralTensor"), const ENeuralTensorTypeGPU InTensorTypeGPU = ENeuralTensorTypeGPU::Generic);
+	explicit FNeuralTensor(const ENeuralDataType InDataType = ENeuralDataType::None, const TArray<int64>& InSizes = TArray<int64>(), const FString& InName = TEXT("FNeuralTensor"), const ENeuralTensorTypeGPU InTensorTypeGPU = ENeuralTensorTypeGPU::Generic);
+
+	/**
+	 * Performance-wise, this constructor makes a deep copy of the data (not optimal). For maximum speed, use the other constructors + GetData()/GetDataCasted().
+	 * It allocates the desired memory, and fills it with the input data from InArray. Equivalent to FNeuralTensor(InVolumeOrSizes) + SetFromArrayCopy(...)
+	 *
+	 * @param InArray The input data to copy from.
+	 * @param InSizes If empty (default), a 1-D tensor of Volume = InArray.Num() will be assumed. If not empty, used to fill the Sizes of this FNeuralTensor.
+	 * @param InName Used for GPU debugging and the ToString() function.
+	 */
+	template <typename T>
+	explicit FNeuralTensor(const TArray<T>& InArray, const TArray<int64>& InSizes = TArray<int64>(), const FString& InName = TEXT("FNeuralTensor"), const ENeuralTensorTypeGPU InTensorTypeGPU = ENeuralTensorTypeGPU::Generic);
+
+	/**
+	 * Copy constructor.
+	 * It performs a deep (slow but safe) copy of the current FNeuralTensor. The resulting FNeuralTensor will not share any parameters with the current one.
+	 */
+	FNeuralTensor(const FNeuralTensor& InTensor);
+
+	/**
+	 * Copy assignment. Analog to FNeuralTensor(const FNeuralTensor& InTensor).
+	 */
+	FNeuralTensor& operator=(const FNeuralTensor& InTensor);
+
+	/**
+	 * Move constructor. It destroys the original FNeuralTensor (InTensor) to be moved for a fast Move construction.
+	 */
+	FNeuralTensor(FNeuralTensor&& InTensor);
+
+	/**
+	 * Move assignment. Analog to FNeuralTensor(FNeuralTensor&& InTensor).
+	 */
+	FNeuralTensor& operator=(FNeuralTensor&& InTensor);
+
+	/**
+	 * Comparison operator (equal). Returns true if the dimensions, sizes, scalar type, and data match with each other.
+	 * It does not consider other properties of the FNeuralTensor (such as ENeuralTensorTypeGPU).
+	 */
+	bool operator==(const FNeuralTensor& InTensorToCopy) const;
+
+	/**
+	 * Comparison operator (not equal). Returns the opposite than the equal (==) operator.
+	 * It does not consider other properties of the FNeuralTensor (such as ENeuralTensorTypeGPU).
+	 */
+	FORCEINLINE bool operator!=(const FNeuralTensor& InTensorToCopy) const;
+
+	/**
+	 * Returns the reference to an element at the given index.
+	 * @returns Reference to the indexed element.
+	 */
+	template <typename T, typename TInput>
+	T& At(const TInput InIndex);
+
+	/**
+	 * Const version of the non-const function At().
+	 * @returns Reference to the indexed element.
+	 */
+	template <typename T, typename TInput>
+	const T& At(const TInput InIndex) const;
+
+	/**
+	 * 2 different functions to get the output tensor results:
+	 * - GetArrayCopy<T> (slower but safer): It returns a copy of the results as a TArray<T>. T has to be the same size than sizeof(DataType).
+	 * - GetUnderlyingUInt8ArrayRef (faster): It returns a reference to the TArray<uint8> that contains the results. The returned reference will only be valid as long as
+	 *   this class instance does not go out of scope and its internal TArray<uint8> is not re-initialized.
+	 * @see GetArrayCopy(), GetUnderlyingUInt8ArrayRef()
+	 */
+	template <typename T>
+	TArray<T> GetArrayCopy() const;
+	FORCEINLINE const TArray<uint8>& GetUnderlyingUInt8ArrayRef() const;
+
+	/**
+	 * Dangerous: If FNeuralTensor goes out of scope, this returned pointer will no longer be valid.
+	 * Analog to TArray::GetData().
+	 * @returns Void pointer to the FNeuralTensor data.
+	 */
+	FORCEINLINE void* GetData();
+
+	/**
+	 * Const analog to TArray::GetData().
+	 * @returns Const void pointer to the FNeuralTensor data.
+	 */
+	FORCEINLINE const void* const GetData() const;
+
+	template<typename T>
+	FORCEINLINE T* GetDataCasted();
+
+	template<typename T>
+	FORCEINLINE const T* const GetDataCasted() const;
+
+	FORCEINLINE FString GetName() const;
+
+	FORCEINLINE ENeuralDataType GetDataType() const;
+
+	/**
+	 * It returns the size of the current dimension, or 1 if InDimension > GetNumberDimensions().
+	 */
+	int64 GetSize(const int32 InDimension) const;
+
+	FORCEINLINE const TArray<int64>& GetSizes() const;
+
+	FORCEINLINE int32 GetNumberDimensions() const;
+
+	/**
+	 * NumInBytes() = Num() * sizeof(type used)
+	 */
+	FORCEINLINE int64 NumInBytes() const;
+
+	/**
+	 * Num() is analog to TArray<T>::Num().
+	 */
+	FORCEINLINE int64 Num() const;
+
+	/**
+	 * IsEmpty() is analog to TArray<T>::IsEmpty().
+	 */
+	FORCEINLINE bool IsEmpty() const;
+
+	/**
+	 * It returns the ENeuralTensorTypeGPU (Generic, Input, Intermediate(Not)Initialized, Output, Weight, etc.).
+	 */
+	FORCEINLINE ENeuralTensorTypeGPU GetTensorTypeGPU() const;
+
+	/**
+	 * It sets the ENeuralTensorTypeGPU (Generic, Input, Intermediate(Not)Initialized, Output, Weight, etc.).
+	 * If the GPU memory was already initialized, it will also UE_LOG a warning.
+	 */
+	void SetTensorTypeGPU(const ENeuralTensorTypeGPU InTensorTypeGPU);
+
+	FORCEINLINE void SetEnableGPU(const bool bInEnableGPU);
+
+	/**
+	 * It uploads/downloads the memory from/to the CPU to/from the GPU based on TensorTypeGPU (which sets a preset of EBufferUsageFlags flags).
+	 * @param InEBufferUsageFlags gives the user total control over the buffer flags (and ignores the TensorTypeGPU flag). This is meant to be filled with a combination of EBufferUsageFlags values.
+	 */
+	void ToGPU_RenderThread(FRDGBuilder* InOutGraphBuilder);
+	void ToGPU_RenderThread(FRDGBuilder* InOutGraphBuilder, const EBufferUsageFlags InEBufferUsageFlags, const bool bInShouldCopyFromCPU);
+	void UpdateSRVAndOrUAV_RenderThread(FRDGBuilder* InOutGraphBuilder);
+	void ToCPU_RenderThread(FRDGBuilder* InOutGraphBuilder);
+
+	TRefCountPtr<class FRDGPooledBuffer>& GetPooledBuffer() const;
+	const FRDGBufferSRVRef GetBufferSRVRef() const;
+	FRDGBufferUAVRef GetBufferUAVRef();
+
+	/**
+	 * Resize the FNeuralTensor to the desired new size.
+	 * Analog to TArray<T>::SetNumUninitialized().
+	 * @param InTensor Set Sizes and DataType from the input InTensor.
+	 * @param InVolume Set InVolume to 0 if memory allocation is not required or the final size is unknown, SetNumUninitialized(InVolume) == SetNumUninitialized({InVolume}).
+	 * @param InSizes Set InSizes to empty if memory allocation is not required or the final size is unknown.
+	 * @param InDataType set to None means that it will maintain the previous type.
+	 */
+	void SetNumUninitialized(const FNeuralTensor& InTensor, const bool bInAllowShrinking = true);
+	FORCEINLINE void SetNumUninitialized(const int64 InVolume, const ENeuralDataType InDataType, const bool bInAllowShrinking = true);
+	void SetNumUninitialized(const TArray<int64>& InSizes, const ENeuralDataType InDataType, const bool bInAllowShrinking = true);
+
+	/**
+	 * This will replace the TArray with the input one, by deeply copying the array (safer and easier to use).
+	 * For maximum performance:
+	 * - If you are given an input TArray (or a FNeuralTensor with different Sizes) that cannot be moved, use SetFromArrayCopy().
+	 * - If you are initializing the FNeuralTensor iteratively (and to avoid a copy), fill the tensor with GetDataCasted().
+	 * Modifying the array InArray after calling this function will not modify the internal FNeuralTensor array.
+	 * The size of both (InArray and current FNeuralTensor) must match, i.e., Num() must match InArray.Num().
+	 */
+	void SetFromUnderlyingUInt8ArrayCopy(const TArray<uint8>& InArray);
+	template<typename T>
+	void SetFromArrayCopy(const TArray<T>& InArray);
+
+	/**
+	 * It will fill the current neural tensor with the input FTensorProto.
+	 * @returns Whether the conversion was successful.
+	 */
+	bool SetFromTensorProto(const FTensorProto* const InTensorProto, const FString& InTensorName, const ENeuralTensorTypeGPU InTensorTypeGPU);
+
+	/**
+	 * It sets all the elements of the FNeuralTensor to InValue.
+	 * It use a double typename to avoid the mistake of SetTo(0) for a double or float (because that 0 would be an int otherwise).
+	 */
+	template<typename T, typename TInput>
+	void SetTo(const TInput InValue);
+
+	/**
+	 * It flips the InDimension dimension of the tensor.
+	 * @return Whether the flip was successful.
+	 */
+	bool Flip(const int32 InDimension);
+
+	/**
+	 * It flips all the dimensions of the tensor in the range [InDimensionFirst, InDimensionLast). Needed for efficient convolution.
+	 * @return Whether the flip was successful.
+	 */
+	bool Flip(const int32 InDimensionFirst, const int32 InDimensionLast);
+
+	/**
+	 * It transposes the matrix (if the tensor has up to 2 dimensions).
+	 * @return Whether the transpose was successful.
+	 */
+	bool Transpose();
+
+	/**
+	 * If Num() is constant, it reshapes the current tensor. I.e., it just updates Sizes with InSizes.
+	 * Reshape() copies and ReshapeMove() moves InSizes.
+	 */
+	bool Reshape(const TArray<int64>& InSizes);
+	bool ReshapeMove(TArray<int64>& InSizes);
+
+	/**
+	 * It returns a FString with up to InMaxNumberElementsToDisplay elements displayed. If InMaxNumberElementsToDisplay <= 0, it displays them all.
+	 * @param bInReturnOnlyData If false (default), it will print all information. E.g., "FNeuralTensor: Int64, Generic, volume=3, sizes={3}, data=[1 2 3]".
+	 * If false, it will simply print the data. I.e., "[1 2 3]" for the previous example.
+	 */
+	FString ToString(const int64 InMaxNumberElementsToDisplay = -1, const bool bInReturnOnlyData = false) const;
+
+protected:
+	/** Used for GPU debugging and the ToString() function. */
+	UPROPERTY(VisibleAnywhere, Category = "Neural Network Inference")
+	FString Name;
+	/** General variables and properties */
+	UPROPERTY(VisibleAnywhere, Category = "Neural Network Inference")
+	ENeuralDataType DataType;
+	UPROPERTY(VisibleAnywhere, Category = "Neural Network Inference")
+	int64 Volume;
+	UPROPERTY(VisibleAnywhere, Category = "Neural Network Inference")
+	TArray<int64> Sizes;
+	UPROPERTY(VisibleAnywhere, Category = "Neural Network Inference")
+	ENeuralTensorTypeGPU TensorTypeGPU;
+
+private:
+	/** CPU-based variables */
+	UPROPERTY()
+	TArray<uint8> ArrayCPU;
+	/**
+	 * GPU-based variables that are transient (do not need to be serialized) and initialized on the fly (with ToGPU/ToCPU).
+	 * - bEnableGPU: By default false, meaning all GPU memory will be disabled and those functions will not do anything. Enable to allow using the GPU functions and variables of the tensor.
+	 * - ArrayCPUForGPUAs32Data: If ArrayCPU is meant for 64-byte data (i.e., int64, uint64, double).
+	 */
+	bool bEnableGPU;
+	TSharedPtr<TRefCountPtr<class FRDGPooledBuffer>> PooledBuffer;
+	TSharedPtr<FRDGBufferSRVRef> BufferSRVRef;
+	TSharedPtr<FRDGBufferUAVRef> BufferUAVRef;
+	TArray<uint8> ArrayCPUForGPUAs32Data;
+
+	/**
+	 * Auxiliary for the templated constructor FNeuralTensor(const ENeuralDataType InDataType, const TArray<T>& InArray, ...).
+	 * InSizeOfT and sizeof(InDataType) should match, as well as the volume of InValueNum and InSizes.
+	 */
+	explicit FNeuralTensor(const ENeuralDataType InDataType, const void* const InValues, const int64 InSizeOfT, const int64 InValueNum, const TArray<int64>& InSizes, const FString& InName, const ENeuralTensorTypeGPU InTensorTypeGPU);
+
+	/**
+	 * Auxiliary for the templated SetFromArrayCopy().
+	 */
+	void SetFromPointer(const void* const InData, const int64 InSizeOfT, const int64 InDataSize);
+
+	/**
+	 * Checks and warns whether the current data type T is incompatible with DataType.
+	 */
+	template<typename T>
+	bool CheckTAndDataType() const;
+	bool CheckTAndDataTypeResult(const bool bInCheckTAndDataTypeResult, const int64 InSizeOfT) const;
+};
+
+
+
+/* FNeuralTensor inline functions
+ *****************************************************************************/
+
+template <typename T>
+FNeuralTensor::FNeuralTensor(const TArray<T>& InArray, const TArray<int64>& InSizes, const FString& InName, const ENeuralTensorTypeGPU InTensorTypeGPU)
+	: FNeuralTensor(FDataType::GetDataType<T>(), InArray.GetData(), sizeof(T), InArray.Num(), (InSizes.Num() > 0 ? InSizes : TArray<int64>({ (int64)InArray.Num() })), InName, InTensorTypeGPU)
+{}
+
+bool FNeuralTensor::operator!=(const FNeuralTensor& InTensorToCopy) const
+{
+	return !(*this == InTensorToCopy);
+}
+
+template <typename T, typename TInput>
+T& FNeuralTensor::At(const TInput InIndex)
+{
+	checkf(CheckTAndDataType<T>(), TEXT("FNeuralTensor::At(): CheckTAndDataType failed."));
+	return ((T*)ArrayCPU.GetData())[InIndex];
+}
+
+template <typename T, typename TInput>
+const T& FNeuralTensor::At(const TInput InIndex) const
+{
+	checkf(CheckTAndDataType<T>(), TEXT("FNeuralTensor::At(): CheckTAndDataType failed."));
+	return ((T*)ArrayCPU.GetData())[InIndex];
+}
+
+template <typename T>
+TArray<T> FNeuralTensor::GetArrayCopy() const
+{
+	TArray<T> Array;
+	if (CheckTAndDataType<T>())
+	{
+		Array.SetNumUninitialized(Num());
+		FMemory::Memcpy(Array.GetData(), ArrayCPU.GetData(), NumInBytes());
+	}
+	return Array;
+}
+
+const TArray<uint8>& FNeuralTensor::GetUnderlyingUInt8ArrayRef() const
+{
+	return ArrayCPU;
+}
+
+void* FNeuralTensor::GetData()
+{
+	return (void*)ArrayCPU.GetData();
+}
+
+const void* const FNeuralTensor::GetData() const
+{
+	return (void*)ArrayCPU.GetData();
+}
+
+template<typename T>
+T* FNeuralTensor::GetDataCasted()
+{
+	checkf(CheckTAndDataType<T>(), TEXT("FNeuralTensor::GetDataCasted(): CheckTAndDataType failed."));
+	return (T*)ArrayCPU.GetData();
+}
+
+template<typename T>
+const T* const FNeuralTensor::GetDataCasted() const
+{
+	checkf(CheckTAndDataType<T>(), TEXT("FNeuralTensor::GetDataCasted(): CheckTAndDataType failed."));
+	return (T*)ArrayCPU.GetData();
+}
+
+FString FNeuralTensor::GetName() const
+{
+	return Name;
+}
+
+ENeuralDataType FNeuralTensor::GetDataType() const
+{
+	return DataType;
+}
+
+const TArray<int64>& FNeuralTensor::GetSizes() const
+{
+	return Sizes;
+}
+
+int32 FNeuralTensor::GetNumberDimensions() const
+{
+	return Sizes.Num();
+}
+
+int64 FNeuralTensor::NumInBytes() const
+{
+	return ArrayCPU.Num();
+}
+
+int64 FNeuralTensor::Num() const
+{
+	return Volume;
+}
+
+bool FNeuralTensor::IsEmpty() const
+{
+	return ArrayCPU.IsEmpty();
+}
+
+ENeuralTensorTypeGPU FNeuralTensor::GetTensorTypeGPU() const
+{
+	return TensorTypeGPU;
+}
+
+void FNeuralTensor::SetEnableGPU(const bool bInEnableGPU)
+{
+	bEnableGPU = bInEnableGPU;
+}
+
+void FNeuralTensor::SetNumUninitialized(const int64 InVolume, const ENeuralDataType InDataType, const bool bInAllowShrinking)
+{
+	SetNumUninitialized((InVolume > 0 ? TArray<int64>({ InVolume }) : TArray<int64>({})), InDataType, bInAllowShrinking);
+}
+
+template<typename T>
+void FNeuralTensor::SetFromArrayCopy(const TArray<T>& InArray)
+{
+	if (CheckTAndDataType<T>())
+	{
+		SetFromPointer(InArray.GetData(), sizeof(T), InArray.Num());
+	}
+}
+
+template<typename T, typename TInput>
+void FNeuralTensor::SetTo(const TInput InValue)
+{
+	if (CheckTAndDataType<T>())
+	{
+		std::fill(GetDataCasted<T>(), GetDataCasted<T>() + Num(), (T)InValue); // There is no UE equivalent for std::fill
+	}
+}
+
+template<typename T>
+bool FNeuralTensor::CheckTAndDataType() const
+{
+	return CheckTAndDataTypeResult(FDataType::CheckTAndDataType<T>(DataType), sizeof(T));
+}
diff --git a/Engine/Plugins/Experimental/NeuralNetworkInference/Source/NeuralNetworkInferenceQA/NeuralNetworkInferenceQA.Build.cs b/Engine/Plugins/Experimental/NeuralNetworkInference/Source/NeuralNetworkInferenceQA/NeuralNetworkInferenceQA.Build.cs
--- a/Engine/Plugins/Experimental/NeuralNetworkInference/Source/NeuralNetworkInferenceQA/NeuralNetworkInferenceQA.Build.cs
+++ b/Engine/Plugins/Experimental/NeuralNetworkInference/Source/NeuralNetworkInferenceQA/NeuralNetworkInferenceQA.Build.cs
@@ -21,7 +21,7 @@ public class NeuralNetworkInferenceQA : ModuleRules
 				"Core",
 				"CoreUObject",
 				"Engine",
-				"NeuralNetworkInferenceBackEnd"
+				"NeuralNetworkInference"
 			}
 		);
 
@@ -29,7 +29,6 @@ public class NeuralNetworkInferenceQA : ModuleRules
 			(
 			new string[] {
 				"ModelProto",
-				"NeuralNetworkInference",
 				"NeuralNetworkInferenceCore",
 				"NeuralNetworkInferenceLegacy"
 			}
