NeuralNetworkInference (NNI) plugin:
- NNI module cleaned: FNOperators, NeuralInt64ArrayUInt32Buffer as part of Internal.

#rnx

diff --git a/Engine/Plugins/Experimental/NeuralNetworkInference/Source/NeuralNetworkInference/Internal/NeuralInt64ArrayUInt32Buffer.h b/Engine/Plugins/Experimental/NeuralNetworkInference/Source/NeuralNetworkInference/Internal/NeuralInt64ArrayUInt32Buffer.h
new file mode 100644
--- /dev/null
+++ b/Engine/Plugins/Experimental/NeuralNetworkInference/Source/NeuralNetworkInference/Internal/NeuralInt64ArrayUInt32Buffer.h
@@ -0,0 +1,133 @@
+// Copyright Epic Games, Inc. All Rights Reserved.
+
+#pragma once
+
+#include "CoreMinimal.h"
+#include "RHI.h"
+
+/**
+ * HLSL uses 32-bit arrays, while FNeuralTensor uses 64-bit ones. This auxiliary class wraps them both together to make sure we
+ * do not modify one type without modifying the other one.
+ * There are 2 ways to use this class, depending on the value of bShouldCreate64BitVersion in Move(TArray<uint32>& InUInt32Array, const bool bShouldCreate64BitVersion):
+ * 1. If bShouldCreate64BitVersion = true, both the 32 and 64 bit TArrays will be created. Useful when both will be read.
+ * 2. If bShouldCreate64BitVersion = false, only the 32 bit TArray will be filled. Useful when the 64 bit one will not be used.
+ * GPU-wise, UInt32ReadBuffer will only be filled if/when CreateAndLoadSRVBuffer() is called. You can use IsReadBuffer() to know if that has happened yet.
+ */
+class NEURALNETWORKINFERENCE_API FNeuralInt64ArrayUInt32Buffer
+{
+public:
+	FNeuralInt64ArrayUInt32Buffer() = default;
+	/**
+	 * Equivalent to calling Copy().
+	 */
+	FNeuralInt64ArrayUInt32Buffer(const TArray<int64>& InInt64Array);
+	FNeuralInt64ArrayUInt32Buffer(const TArray<uint32>& InUInt32Array, const bool bShouldCreate64BitVersion);
+
+	FORCEINLINE int32 Num() const;
+
+	FORCEINLINE bool IsEmpty() const;
+
+	FORCEINLINE const TArray<uint32>& GetUInt32() const;
+
+	/**
+	 * In addition to return the TArray, it also checks whether bShouldCreate64BitVersion was true when Move() was called.
+	 */
+	const TArray<int64>& GetInt64() const;
+
+	FORCEINLINE const FShaderResourceViewRHIRef& GetUInt32SRV() const;
+
+	FORCEINLINE bool IsReadBufferValid() const;
+
+	FORCEINLINE void ReleaseReadBuffer();
+
+	FORCEINLINE void ResetReadBuffer();
+
+	/**
+	 * Create and TArray::Init() UInt32Array and optionally Int64Array.
+	 * 32-bit version is always generated, 64 only if bShouldCreate64BitVersion = true.
+	 */
+	void Init(const uint32 InArrayValue, const uint32 InArraySize, const bool bShouldCreate64BitVersion);
+
+	/**
+	 * Create and fill UInt32Array and optionally Int64Array.
+	 * 32-bit version is always generated, 64 only if bShouldCreate64BitVersion = true.
+	 * TArray will be moved for performance in Move (thus destroyed).
+	 */
+	void Copy(const TArray<int64>& InInt64Array);
+	void Copy(const TArray<uint32>& InUInt32Array, const bool bShouldCreate64BitVersion);
+	void Move(TArray<int64>& InInt64Array);
+	void Move(TArray<uint32>& InUInt32Array, const bool bShouldCreate64BitVersion);
+
+	/**
+	 *  It sets and fills UInt32ReadBuffer with UInt32Array.
+	 */
+	void CreateAndLoadSRVBuffer(const TCHAR* InDebugName);
+
+	FORCEINLINE int64 operator[](const int64 InIndex) const;
+	/**
+	 * DO NOT USE DIRECTLY
+	 * STL-like iterators to enable range-based for loop support.
+	 */
+	FORCEINLINE const int64* begin() const;
+	FORCEINLINE const int64* end() const;
+
+private:
+	TArray<int64> Int64Array;
+	TArray<uint32> UInt32Array;
+	TSharedPtr<FReadBuffer> UInt32ReadBuffer;
+};
+
+
+
+/* FNeuralInt64ArrayUInt32Buffer inline functions
+ *****************************************************************************/
+
+int32 FNeuralInt64ArrayUInt32Buffer::Num() const
+{
+	return UInt32Array.Num();
+}
+
+bool FNeuralInt64ArrayUInt32Buffer::IsEmpty() const
+{
+	return UInt32Array.IsEmpty();
+}
+
+const TArray<uint32>& FNeuralInt64ArrayUInt32Buffer::GetUInt32() const
+{
+	return UInt32Array;
+}
+
+const FShaderResourceViewRHIRef& FNeuralInt64ArrayUInt32Buffer::GetUInt32SRV() const
+{
+	return UInt32ReadBuffer->SRV;
+}
+
+bool FNeuralInt64ArrayUInt32Buffer::IsReadBufferValid() const
+{
+	return UInt32ReadBuffer.IsValid();
+}
+
+void FNeuralInt64ArrayUInt32Buffer::ReleaseReadBuffer()
+{
+	UInt32ReadBuffer->Release();
+}
+
+void FNeuralInt64ArrayUInt32Buffer::ResetReadBuffer()
+{
+	UInt32ReadBuffer.Reset();
+}
+
+int64 FNeuralInt64ArrayUInt32Buffer::operator[](const int64 InIndex) const
+{
+	return Int64Array[InIndex];
+}
+
+const int64* FNeuralInt64ArrayUInt32Buffer::begin() const
+{
+	return Int64Array.GetData();
+}
+
+const int64* FNeuralInt64ArrayUInt32Buffer::end() const
+{
+	return Int64Array.GetData() + Int64Array.Num();
+}
diff --git a/Engine/Plugins/Experimental/NeuralNetworkInference/Source/NeuralNetworkInference/Internal/NeuralOperator.h b/Engine/Plugins/Experimental/NeuralNetworkInference/Source/NeuralNetworkInference/Internal/NeuralOperator.h
new file mode 100644
--- /dev/null
+++ b/Engine/Plugins/Experimental/NeuralNetworkInference/Source/NeuralNetworkInference/Internal/NeuralOperator.h
@@ -0,0 +1,204 @@
+// Copyright Epic Games, Inc. All Rights Reserved.
+
+#pragma once
+
+#include "CoreMinimal.h"
+#include "NeuralTensor.h"
+
+class FRDGBuilder;
+
+/**
+ * This is an auxiliary class. See UNeuralNetworkLegacy for a high-level wrapper of the whole NeuralNetworkInference plugin. The UNeuralNetworkLegacy header
+ * documentation also includes some code examples.
+ *
+ * FNeuralOperator represents an operator (or layer) of the UNeuralNetworkLegacy.
+ */
+class NEURALNETWORKINFERENCE_API FNeuralOperator
+{
+
+public:
+	/**
+	 * @return InInlinedTensor Tensor that will be inlined, EstimateInlinedTensorFromPotentialOnes() will make no effect here.
+	 * @return InPotentialInlinedTensors Potential tensors that could be inlined. EstimateInlinedTensorFromPotentialOnes() should be implemented by subclasses that use this constructor to decide the tensor that will be inlined.
+	 */
+	FNeuralOperator(const FString& InName, const int32 InVersion, const int32 InInlinedTensor);
+	FNeuralOperator(const FString& InName, const int32 InVersion, const TSet<uint32>& InPotentialInlinedTensors);
+
+	virtual ~FNeuralOperator();
+
+	FORCEINLINE virtual FString GetName() const final;
+
+	/**
+	 * It returns the maximum possible number of auxiliary tensors per operator. So far, the worse case is 1 (for FConvOperator and FConvTranposeOperator).
+	 * Modify 1 to a higher value if a new operator with a higher number of auxiliary tensors appear.
+	 */
+	FORCEINLINE static int32 GetMaximumPossibleNumberAuxiliaryTensorsPerOperator();
+
+	/**
+	 * By default, it will return 0 (meaning there is no need for auxiliary tensors).
+	 * However, some operators (e.g., FConvOperator and FConvTransposeOperator) might need auxiliary tensors. If so, this function must be re-implemented and return the number of auxiliary tensors needed.
+	 */
+	FORCEINLINE virtual int32 GetNumberAuxiliaryTensors() const;
+
+	/**
+	 * It must be called for each operator before setting its output tensors and before calling Forward.
+	 * It fills the input tensor(s) of this operator and configures it.
+	 */
+	virtual int32 SetInputTensorsAndGetInlinedIndex(const TArray<FNeuralTensor*>& InInputTensors) final;
+
+	/**
+	 * It must be called for operator that needs to define and use auxiliary tensors (e.g., FConvOperator and FConvTransposeOperator).
+	 */
+	virtual void SetAuxiliaryTensors(const TArray<FNeuralTensor*>& InAuxiliaryTensors) final;
+
+	/**
+	 * It must be called for each non-fully-inlined operator after setting its input tensors and before calling Forward.
+	 */
+	FORCEINLINE virtual void SetOutputTensors(const TArray<FNeuralTensor*>& InOutputTensors) final;
+
+	/**
+	 * It must be called for each operator after setting its input and output tensors and before calling Forward.
+	 * It internally runs:
+	 * - ConfigureOutputAndInternalVariablesAndSanityChecks(): It configures the output and auxiliary variables as well as run sanity checks to make sure input, output, and auxiliary variables all have compatible values.
+	 * - SetWhetherXXXTensorsNeedTransferToGPUForGPUMode(): By default, these functions just sets SetEnableGPU(true) for all input, output, and auxiliary tensors. However, some operators (e.g.,
+	 *   FReshapeOperator) override this function and only SetEnableGPU(true) to a subset of them. E.g., FReshape do not SetEnableGPU(true) the 2nd input tensor (the shape tensor).
+	 *     - SetWhetherInputTensorsNeedTransferToGPUForGPUMode().
+	 *     - SetWhetherOutputTensorsNeedTransferToGPUForGPUMode().
+	 *     - SetWhetherAuxiliaryTensorsNeedTransferToGPUForGPUMode().
+	 *
+	 * @return Whether this operator was properly set and configured.
+	 */
+	virtual bool Configure() final;
+
+	/**
+	 * It should be implemented for those FNeuralOperator child classes that require to allocate auxiliary GPU memory.
+	 * Most operators do not require this, so by default it will do nothing (not purely virtual).
+	 */
+	FORCEINLINE virtual void ToGPU_RenderThread();
+
+	/**
+	 * It runs the forward pass on the CPU/GPU (ForwardCPU/ForwardGPU_RenderThread). The CPU one will block the current thread until it finishes and
+	 * the GPU one must be called from the render thread. By default, most child classes will not needed it.
+	 */
+	virtual void ForwardCPU() = 0;
+	virtual void ForwardGPU_RenderThread(FRDGBuilder* InOutGraphBuilder) = 0;
+
+	/**
+	 * It runs the post forward pass on the CPU/GPU (PostForwardCPU/PostForwardGPU_RenderThread). The CPU one will block the current thread until it
+	 * finishes and the GPU one must be called from the render thread. By default, most child classes will not needed it.
+	 */
+	FORCEINLINE virtual bool NeedsPostForward() const;
+	FORCEINLINE virtual void PostForwardCPU();
+	FORCEINLINE virtual void PostForwardGPU_RenderThread(FRDGBuilder* InOutGraphBuilder);
+
+	virtual FString ToString() const;
+
+protected:
+	bool bIsLoaded;
+	const FString Name;
+	const int32 Version;
+	const TSet<uint32> PotentiallyInlinedTensors;
+	int32 InlinedTensor;
+	TArray<FNeuralTensor*> AuxiliaryTensors; /* Only for some operators that might need them (e.g., FConvOperator and FConvTransposeOperator). */
+
+	/**
+	 * 2 cases:
+	 * - If PotentiallyInlinedTensors is empty, this class should not be re-implemented by the child classes.
+	 * - If PotentiallyInlinedTensors is not empty, then it must be implemented by the child classes.
+	 *
+	 * If not empty, it decides which input tensor can be used as output (if the operator can be inlined), or -1 otherwise.
+	 * There are 3 type of operators:
+	 *  1. Those that cannot be inlined. This kind of operators do not need to implement EstimateInlinedTensorFromPotentialOnes() nor modify InlinedTensor.
+	 *  2. Those that can be inlined. This kind of operators do not need to implement EstimateInlinedTensorFromPotentialOnes() but can instead modify InlinedTensor on construction time.
+	 *  3. Those that can be inlined but only if the input tensors satisfy certain constraints (e.g., size). These must re-implement this function.
+	 */
+	FORCEINLINE virtual bool EstimateInlinedTensorFromPotentialOnes();
+
+	/**
+	 * See Configure() for a full explanation of these.
+	 */
+	virtual bool ConfigureOutputAndInternalVariablesAndSanityChecks() = 0;
+	virtual bool SetWhetherInputTensorsNeedTransferToGPUForGPUMode();
+	virtual bool SetWhetherOutputTensorsNeedTransferToGPUForGPUMode();
+	virtual bool SetWhetherAuxiliaryTensorsNeedTransferToGPUForGPUMode();
+
+	/**
+	 * GetInputTensorNoConst() is GetInputTensorsNoConst() for single output operators. It will log a warning if used with a multi-output operator.
+	 */
+	virtual FNeuralTensor& GetInputTensorNoConst() final;
+	virtual TArray<FNeuralTensor*>& GetInputTensorsNoConst() final;
+
+	/**
+	 * GetInputTensorConst() is GetInputTensorsConst() for single output operators. It will log a warning if used with a multi-output operator.
+	 */
+	virtual const FNeuralTensor& GetInputTensorConst() const final;
+	FORCEINLINE virtual const TArray<FNeuralTensor*>& GetInputTensorsConst() const final;
+
+	/**
+	 * GetOutputTensorNoConst() is GetOutputTensorsNoConst() for single output operators. It will log a warning if used with a multi-output operator.
+	 */
+	virtual FNeuralTensor& GetOutputTensorNoConst() final;
+	virtual TArray<FNeuralTensor*>& GetOutputTensorsNoConst() final;
+
+	/**
+	 * GetOutputTensorConst() is GetOutputTensorsConst() for single output operators. It will log a warning if used with a multi-output operator.
+	 */
+	virtual const FNeuralTensor& GetOutputTensorConst() const final;
+	virtual const TArray<FNeuralTensor*>& GetOutputTensorsConst() const final;
+
+private:
+	TArray<FNeuralTensor*> PrivateInputTensors; /* Const for non-inlined operators. */
+	TArray<FNeuralTensor*> PrivateOutputTensors; /* Only for non-inlined operators. */
+};
+
+
+
+/* FNeuralOperator inline functions
+ *****************************************************************************/
+
+FString FNeuralOperator::GetName() const
+{
+	return Name;
+}
+
+int32 FNeuralOperator::GetMaximumPossibleNumberAuxiliaryTensorsPerOperator()
+{
+	return 1;
+}
+
+int32 FNeuralOperator::GetNumberAuxiliaryTensors() const
+{
+	return 0;
+}
+
+void FNeuralOperator::SetOutputTensors(const TArray<FNeuralTensor*>& InOutputTensors)
+{
+	PrivateOutputTensors = InOutputTensors;
+}
+
+void FNeuralOperator::ToGPU_RenderThread()
+{
+}
+
+bool FNeuralOperator::NeedsPostForward() const
+{
+	return false;
+}
+
+void FNeuralOperator::PostForwardCPU()
+{
+}
+
+void FNeuralOperator::PostForwardGPU_RenderThread(FRDGBuilder* InOutGraphBuilder)
+{
+}
+
+bool FNeuralOperator::EstimateInlinedTensorFromPotentialOnes()
+{
+	return (PotentiallyInlinedTensors.Num() == 0);
+}
+
+const TArray<FNeuralTensor*>& FNeuralOperator::GetInputTensorsConst() const
+{
+	return PrivateInputTensors;
+}
diff --git a/Engine/Plugins/Experimental/NeuralNetworkInference/Source/NeuralNetworkInference/Internal/NeuralOperators.h b/Engine/Plugins/Experimental/NeuralNetworkInference/Source/NeuralNetworkInference/Internal/NeuralOperators.h
new file mode 100644
--- /dev/null
+++ b/Engine/Plugins/Experimental/NeuralNetworkInference/Source/NeuralNetworkInference/Internal/NeuralOperators.h
@@ -0,0 +1,44 @@
+// Copyright Epic Games, Inc. All Rights Reserved.
+
+#pragma once
+
+// Abstract layers
+#include "NeuralOperator.h"
+#include "NeuralOperators/ElementWiseOperator.h"
+#include "NeuralOperators/MultidirectionalBroadcastOperator.h"
+
+// Non-abstract layers
+#include "NeuralOperators/AbsOperator.h"
+#include "NeuralOperators/AcosOperator.h"
+#include "NeuralOperators/AddOperator.h"
+#include "NeuralOperators/AsinOperator.h"
+#include "NeuralOperators/AtanOperator.h"
+#include "NeuralOperators/BatchNormalizationOperator.h"
+#include "NeuralOperators/CeilOperator.h"
+#include "NeuralOperators/ConstantOperator.h"
+#include "NeuralOperators/ConvOperator.h"
+#include "NeuralOperators/ConvTransposeOperator.h"
+#include "NeuralOperators/CosOperator.h"
+#include "NeuralOperators/CoshOperator.h"
+#include "NeuralOperators/DivOperator.h"
+#include "NeuralOperators/ExpOperator.h"
+#include "NeuralOperators/FloorOperator.h"
+#include "NeuralOperators/GemmOperator.h"
+#include "NeuralOperators/LeakyReluOperator.h"
+#include "NeuralOperators/LogOperator.h"
+#include "NeuralOperators/MulOperator.h"
+#include "NeuralOperators/NegOperator.h"
+#include "NeuralOperators/PowOperator.h"
+#include "NeuralOperators/ReciprocalOperator.h"
+#include "NeuralOperators/ReluOperator.h"
+#include "NeuralOperators/ReshapeOperator.h"
+#include "NeuralOperators/RoundOperator.h"
+#include "NeuralOperators/SigmoidOperator.h"
+#include "NeuralOperators/SignOperator.h"
+#include "NeuralOperators/SinOperator.h"
+#include "NeuralOperators/SinhOperator.h"
+#include "NeuralOperators/SqrtOperator.h"
+#include "NeuralOperators/SqueezeOperator.h"
+#include "NeuralOperators/SubOperator.h"
+#include "NeuralOperators/TanOperator.h"
+#include "NeuralOperators/TanhOperator.h"
diff --git a/Engine/Plugins/Experimental/NeuralNetworkInference/Source/NeuralNetworkInference/Internal/NeuralOperators/AbsOperator.h b/Engine/Plugins/Experimental/NeuralNetworkInference/Source/NeuralNetworkInference/Internal/NeuralOperators/AbsOperator.h
new file mode 100644
--- /dev/null
+++ b/Engine/Plugins/Experimental/NeuralNetworkInference/Source/NeuralNetworkInference/Internal/NeuralOperators/AbsOperator.h
@@ -0,0 +1,26 @@
+// Copyright Epic Games, Inc. All Rights Reserved.
+
+#pragma once
+
+#include "CoreMinimal.h"
+#include "NeuralOperators/ElementWiseOperator.h"
+
+class NEURALNETWORKINFERENCE_API FAbsOperator : public IElementWiseOperator
+{
+public:
+	FAbsOperator(const bool bIsInlinedTensor);
+
+	virtual ~FAbsOperator();
+
+	FORCEINLINE virtual void ForwardCPU() override final;
+};
+
+
+
+/* FAbsOperator inline functions
+ *****************************************************************************/
+
+void FAbsOperator::ForwardCPU()
+{
+	ForwardCPUWithFunction([](const float InValue) { return FMath::Abs(InValue); });
+}
diff --git a/Engine/Plugins/Experimental/NeuralNetworkInference/Source/NeuralNetworkInference/Internal/NeuralOperators/AcosOperator.h b/Engine/Plugins/Experimental/NeuralNetworkInference/Source/NeuralNetworkInference/Internal/NeuralOperators/AcosOperator.h
new file mode 100644
--- /dev/null
+++ b/Engine/Plugins/Experimental/NeuralNetworkInference/Source/NeuralNetworkInference/Internal/NeuralOperators/AcosOperator.h
@@ -0,0 +1,26 @@
+// Copyright Epic Games, Inc. All Rights Reserved.
+
+#pragma once
+
+#include "CoreMinimal.h"
+#include "NeuralOperators/ElementWiseOperator.h"
+
+class NEURALNETWORKINFERENCE_API FAcosOperator : public IElementWiseOperator
+{
+public:
+	FAcosOperator(const bool bIsInlinedTensor);
+
+	virtual ~FAcosOperator();
+
+	FORCEINLINE virtual void ForwardCPU() override final;
+};
+
+
+
+/* FAcosOperator inline functions
+ *****************************************************************************/
+
+void FAcosOperator::ForwardCPU()
+{
+	ForwardCPUWithFunction([](const float InValue) { return FMath::Acos(InValue); });
+}
diff --git a/Engine/Plugins/Experimental/NeuralNetworkInference/Source/NeuralNetworkInference/Internal/NeuralOperators/AddOperator.h b/Engine/Plugins/Experimental/NeuralNetworkInference/Source/NeuralNetworkInference/Internal/NeuralOperators/AddOperator.h
new file mode 100644
--- /dev/null
+++ b/Engine/Plugins/Experimental/NeuralNetworkInference/Source/NeuralNetworkInference/Internal/NeuralOperators/AddOperator.h
@@ -0,0 +1,26 @@
+// Copyright Epic Games, Inc. All Rights Reserved.
+
+#pragma once
+
+#include "CoreMinimal.h"
+#include "NeuralOperators/MultidirectionalBroadcastOperator.h"
+
+class NEURALNETWORKINFERENCE_API FAddOperator : public IMultidirectionalBroadcastOperator
+{
+public:
+	FAddOperator(const TSet<uint32>& InPotentialInlinedTensors);
+
+	virtual ~FAddOperator();
+
+	FORCEINLINE virtual void ForwardCPU() override final;
+};
+
+
+
+/* FAddOperator inline functions
+ *****************************************************************************/
+
+void FAddOperator::ForwardCPU()
+{
+	ForwardCPUWithFunction([](const float InValueA, const float InValueB) { return InValueA + InValueB; });
+}
diff --git a/Engine/Plugins/Experimental/NeuralNetworkInference/Source/NeuralNetworkInference/Internal/NeuralOperators/AsinOperator.h b/Engine/Plugins/Experimental/NeuralNetworkInference/Source/NeuralNetworkInference/Internal/NeuralOperators/AsinOperator.h
new file mode 100644
--- /dev/null
+++ b/Engine/Plugins/Experimental/NeuralNetworkInference/Source/NeuralNetworkInference/Internal/NeuralOperators/AsinOperator.h
@@ -0,0 +1,26 @@
+// Copyright Epic Games, Inc. All Rights Reserved.
+
+#pragma once
+
+#include "CoreMinimal.h"
+#include "NeuralOperators/ElementWiseOperator.h"
+
+class NEURALNETWORKINFERENCE_API FAsinOperator : public IElementWiseOperator
+{
+public:
+	FAsinOperator(const bool bIsInlinedTensor);
+
+	virtual ~FAsinOperator();
+
+	FORCEINLINE virtual void ForwardCPU() override final;
+};
+
+
+
+/* FAsinOperator inline functions
+ *****************************************************************************/
+
+void FAsinOperator::ForwardCPU()
+{
+	ForwardCPUWithFunction([](const float InValue) { return FMath::Asin(InValue); });
+}
diff --git a/Engine/Plugins/Experimental/NeuralNetworkInference/Source/NeuralNetworkInference/Internal/NeuralOperators/AtanOperator.h b/Engine/Plugins/Experimental/NeuralNetworkInference/Source/NeuralNetworkInference/Internal/NeuralOperators/AtanOperator.h
new file mode 100644
--- /dev/null
+++ b/Engine/Plugins/Experimental/NeuralNetworkInference/Source/NeuralNetworkInference/Internal/NeuralOperators/AtanOperator.h
@@ -0,0 +1,26 @@
+// Copyright Epic Games, Inc. All Rights Reserved.
+
+#pragma once
+
+#include "CoreMinimal.h"
+#include "NeuralOperators/ElementWiseOperator.h"
+
+class NEURALNETWORKINFERENCE_API FAtanOperator : public IElementWiseOperator
+{
+public:
+	FAtanOperator(const bool bIsInlinedTensor);
+
+	virtual ~FAtanOperator();
+
+	FORCEINLINE virtual void ForwardCPU() override final;
+};
+
+
+
+/* FAtanOperator inline functions
+ *****************************************************************************/
+
+void FAtanOperator::ForwardCPU()
+{
+	ForwardCPUWithFunction([](const float InValue) { return FMath::Atan(InValue); });
+}
diff --git a/Engine/Plugins/Experimental/NeuralNetworkInference/Source/NeuralNetworkInference/Internal/NeuralOperators/BatchNormalizationOperator.h b/Engine/Plugins/Experimental/NeuralNetworkInference/Source/NeuralNetworkInference/Internal/NeuralOperators/BatchNormalizationOperator.h
new file mode 100644
--- /dev/null
+++ b/Engine/Plugins/Experimental/NeuralNetworkInference/Source/NeuralNetworkInference/Internal/NeuralOperators/BatchNormalizationOperator.h
@@ -0,0 +1,27 @@
+// Copyright Epic Games, Inc. All Rights Reserved.
+
+#pragma once
+
+#include "CoreMinimal.h"
+#include "ModelProto.h"
+#include "NeuralOperator.h"
+
+class NEURALNETWORKINFERENCE_API FBatchNormalizationOperator : public FNeuralOperator
+{
+public:
+	FBatchNormalizationOperator(const bool bIsInlinedTensor, const FNodeProto& InNodeProto);
+
+	FBatchNormalizationOperator(const bool bIsInlinedTensor, const float InEpsilon = 1e-5, const float InMomentum = 0.9f);
+
+	virtual ~FBatchNormalizationOperator();
+
+	virtual bool ConfigureOutputAndInternalVariablesAndSanityChecks() override final;
+
+	virtual void ForwardCPU() override final;
+
+	virtual void ForwardGPU_RenderThread(FRDGBuilder* InOutGraphBuilder) override final;
+
+private:
+	float Epsilon;
+	float Momentum;
+};
diff --git a/Engine/Plugins/Experimental/NeuralNetworkInference/Source/NeuralNetworkInference/Internal/NeuralOperators/CeilOperator.h b/Engine/Plugins/Experimental/NeuralNetworkInference/Source/NeuralNetworkInference/Internal/NeuralOperators/CeilOperator.h
new file mode 100644
--- /dev/null
+++ b/Engine/Plugins/Experimental/NeuralNetworkInference/Source/NeuralNetworkInference/Internal/NeuralOperators/CeilOperator.h
@@ -0,0 +1,26 @@
+// Copyright Epic Games, Inc. All Rights Reserved.
+
+#pragma once
+
+#include "CoreMinimal.h"
+#include "NeuralOperators/ElementWiseOperator.h"
+
+class NEURALNETWORKINFERENCE_API FCeilOperator : public IElementWiseOperator
+{
+public:
+	FCeilOperator(const bool bIsInlinedTensor);
+
+	virtual ~FCeilOperator();
+
+	FORCEINLINE virtual void ForwardCPU() override final;
+};
+
+
+
+/* FCeilOperator inline functions
+ *****************************************************************************/
+
+void FCeilOperator::ForwardCPU()
+{
+	ForwardCPUWithFunction([](const float InValue) { return FMath::CeilToFloat(InValue); });
+}
diff --git a/Engine/Plugins/Experimental/NeuralNetworkInference/Source/NeuralNetworkInference/Internal/NeuralOperators/ConstantOperator.h b/Engine/Plugins/Experimental/NeuralNetworkInference/Source/NeuralNetworkInference/Internal/NeuralOperators/ConstantOperator.h
new file mode 100644
--- /dev/null
+++ b/Engine/Plugins/Experimental/NeuralNetworkInference/Source/NeuralNetworkInference/Internal/NeuralOperators/ConstantOperator.h
@@ -0,0 +1,26 @@
+// Copyright Epic Games, Inc. All Rights Reserved.
+
+#pragma once
+
+#include "CoreMinimal.h"
+#include "ModelProto.h"
+#include "NeuralOperator.h"
+
+class NEURALNETWORKINFERENCE_API FConstantOperator : public FNeuralOperator
+{
+public:
+	FConstantOperator(const FNodeProto& InNodeProto);
+
+	FConstantOperator(const FNeuralTensor& InTensor);
+
+	virtual ~FConstantOperator();
+
+	virtual bool ConfigureOutputAndInternalVariablesAndSanityChecks() override final;
+
+	virtual void ForwardCPU() override final;
+
+	virtual void ForwardGPU_RenderThread(FRDGBuilder* InOutGraphBuilder) override final;
+
+private:
+	FNeuralTensor Tensor;
+};
diff --git a/Engine/Plugins/Experimental/NeuralNetworkInference/Source/NeuralNetworkInference/Internal/NeuralOperators/ConvBaseOperator.h b/Engine/Plugins/Experimental/NeuralNetworkInference/Source/NeuralNetworkInference/Internal/NeuralOperators/ConvBaseOperator.h
new file mode 100644
--- /dev/null
+++ b/Engine/Plugins/Experimental/NeuralNetworkInference/Source/NeuralNetworkInference/Internal/NeuralOperators/ConvBaseOperator.h
@@ -0,0 +1,130 @@
+// Copyright Epic Games, Inc. All Rights Reserved.
+
+#pragma once
+
+#include "CoreMinimal.h"
+#include "NeuralInt64ArrayUInt32Buffer.h"
+#include "NeuralOperator.h"
+
+class NEURALNETWORKINFERENCE_API FConvBaseOperator : public FNeuralOperator
+{
+public:
+	enum class EAutoPad : uint8
+	{
+		NotSet, /* Explicit padding is used */
+		SameUpper, /* Auto-padding the input so that input and output match in size (if odd number, extra padding at the beginning) */
+		SameLower, /* Auto-padding the input so that input and output match in size (if odd number, extra padding at the end) */
+		Valid /* No padding is used */
+	};
+
+	virtual ~FConvBaseOperator();
+
+	virtual bool ConfigureOutputAndInternalVariablesAndSanityChecks() override;
+
+	virtual void ToGPU_RenderThread() override;
+
+	virtual void ForwardCPU() override;
+
+	virtual void ForwardGPU_RenderThread(FRDGBuilder* InOutGraphBuilder) override;
+
+protected:
+	/**
+	 * Whether it is a normal or transposed convolution.
+	 */
+	const bool bIsTransposed;
+	/**
+	 * Parameters given as constructor argument / part of the child FNodeProto (might or might not be optional)
+	 */
+	EAutoPad AutoPad;
+	FNeuralInt64ArrayUInt32Buffer Dilations;
+	int64 Group;
+	TArray<int64> KernelShape;
+	FNeuralInt64ArrayUInt32Buffer Pads;
+	FNeuralInt64ArrayUInt32Buffer Strides;
+	/**
+	 * Only used on ConvTransposeOperator
+	 */
+	TArray<int64> OutputPadding;
+	TArray<int64> OutputShape;
+	TArray<int64> ConvolutionStridesIfTransposedConvolution;
+	/**
+	 * Internally auto generated. Only for the GPU forward pass.
+	 */
+	FNeuralInt64ArrayUInt32Buffer OutputSizesForGPU;
+	FNeuralInt64ArrayUInt32Buffer XOrXWithZerosSizesForGPU;
+	FNeuralInt64ArrayUInt32Buffer WSizesForGPU;
+
+	/**
+	 * Only meant for inherited classes, which also need to fill the InName/InVersion/InInlinedTensor properties.
+	 * @param InOutputPadding Only for ConvTranspose
+	 */
+	FConvBaseOperator(const FString& InName, const int32 InVersion, const int32 InInlinedTensor,
+		const EAutoPad InAutoPad, const TArray<int64>& InDilations, const int64 InGroup, const TArray<int64>& InKernelShape, const TArray<int64>& InPads, const TArray<int64>& InStrides,
+		const bool bInIsTransposed, const TArray<int64>& InOutputPadding = TArray<int64>(), const TArray<int64>& InOutputShape = TArray<int64>());
+
+	/**
+	 * It fills Strides.
+	 * @return Whether the sanity checks passed.
+	 */
+	virtual bool SetAndConfigureStrides(const int32 InNumberConvolutionalDimensions) = 0;
+
+	/**
+	 * It sets and configures AuxiliaryTensors[0] if needed (i.e., if XOrXWithZeros is not X).
+	 * @return Whether the sanity checks passed.
+	 */
+	FORCEINLINE virtual bool SetAndConfigureAuxiliaryTensor();
+
+private:
+	/**
+	 * Internally auto generated. Only for the CPU forward pass.
+	 * It defines the limits where we can guarantee the input will not go out of bounds so we can avoid the costly if 0 <= X < XSize for those pixels.
+	 */
+	TArray<int64> OutputPaddedMargins;
+
+	/**
+	 * Meant to be called at the very beginning of ConfigureOutputAndInternalVariablesAndSanityChecks().
+	 */
+	virtual bool SanityChecksForTensors(const TArray<int64>& InPartiallySetOutputSizes) const final;
+
+	/**
+	 * It fills OutputShape (if ConvTranspose and AutoPad == SameLower/Upper).
+	 * @return Whether the sanity checks passed.
+	 */
+	virtual bool FillOutputShape(const TArray<int64>& InXSizes) final;
+
+	/**
+	 * It fills Dilations.
+	 * @return Whether the sanity checks passed.
+	 */
+	virtual bool FillDilationsIfEmpty(const int32 InNumberConvolutionalDimensions) final;
+
+	/**
+	 * It fills and returns EffectiveKernelSizes.
+	 */
+	virtual TArray<int64> CreateEffectiveKernelSizes(const int32 InNumberConvolutionalDimensions, const TArray<int64>& InWSizes) final;
+
+	/**
+	 * It fills and returns Pads based on the input InPads E.g., InPads might be empty (assuming 0 padding on each dimension), the returned Pads would then be filled with NumberConvolutionalDimensions 0s.
+	 * @return Whether the sanity checks passed.
+	 */
+	virtual bool EstimateAndFillPads(const TArray<int64>& InEffectiveKernelSizes, const int32 InNumberDimensions, const TArray<int64>& InXSizes) final;
+
+	/**
+	 * Meant to be called at the very end of ConfigureOutputAndInternalVariablesAndSanityChecks().
+	 * It fills both OutputSizes and OutputPaddedMargins.
+	 * @param InOutPartiallySetOutputSizes For efficiency, it will be passed by reference and modified internally.
+	 */
+	virtual bool FillOutputSizesAndPaddedMarginsAndSanityChecks(TArray<int64>& InOutPartiallySetOutputSizes, const TArray<int64>& InEffectiveKernelSizes) final;
+
+	virtual void MultiChannelConvolutionCPU(FNeuralTensor& OutY, const FNeuralTensor& InXOrXWithZerosSizes, const FNeuralTensor& InW) final;
+};
+
+
+
+/* FConvBaseOperator inline functions
+ *****************************************************************************/
+
+bool FConvBaseOperator::SetAndConfigureAuxiliaryTensor()
+{
+	return true;
+}
diff --git a/Engine/Plugins/Experimental/NeuralNetworkInference/Source/NeuralNetworkInference/Internal/NeuralOperators/ConvOperator.h b/Engine/Plugins/Experimental/NeuralNetworkInference/Source/NeuralNetworkInference/Internal/NeuralOperators/ConvOperator.h
new file mode 100644
--- /dev/null
+++ b/Engine/Plugins/Experimental/NeuralNetworkInference/Source/NeuralNetworkInference/Internal/NeuralOperators/ConvOperator.h
@@ -0,0 +1,19 @@
+// Copyright Epic Games, Inc. All Rights Reserved.
+
+#pragma once
+
+#include "ModelProto.h"
+#include "NeuralOperators/ConvBaseOperator.h"
+
+class NEURALNETWORKINFERENCE_API FConvOperator : public FConvBaseOperator
+{
+public:
+	FConvOperator(const FNodeProto& InNodeProto);
+
+	FConvOperator(const EAutoPad InAutoPad, const TArray<int64>& InDilations, const int64 InGroup, const TArray<int64>& InKernelShape, const TArray<int64>& InPads, const TArray<int64>& InStrides);
+
+	virtual ~FConvOperator();
+
+protected:
+	virtual bool SetAndConfigureStrides(const int32 InNumberConvolutionalDimensions) override final;
+};
diff --git a/Engine/Plugins/Experimental/NeuralNetworkInference/Source/NeuralNetworkInference/Internal/NeuralOperators/ConvTransposeOperator.h b/Engine/Plugins/Experimental/NeuralNetworkInference/Source/NeuralNetworkInference/Internal/NeuralOperators/ConvTransposeOperator.h
new file mode 100644
--- /dev/null
+++ b/Engine/Plugins/Experimental/NeuralNetworkInference/Source/NeuralNetworkInference/Internal/NeuralOperators/ConvTransposeOperator.h
@@ -0,0 +1,65 @@
+// Copyright Epic Games, Inc. All Rights Reserved.
+
+#pragma once
+
+#include "ModelProto.h"
+#include "NeuralInt64ArrayUInt32Buffer.h"
+#include "NeuralOperators/ConvBaseOperator.h"
+
+class NEURALNETWORKINFERENCE_API FConvTransposeOperator : public FConvBaseOperator
+{
+public:
+	FConvTransposeOperator(const FNodeProto& InNodeProto);
+
+	FConvTransposeOperator(const EAutoPad InAutoPad, const TArray<int64>& InDilations, const int64 InGroup, const TArray<int64>& InKernelShape,
+		const TArray<int64>& InOutputPadding, const TArray<int64>& InOutputShape, const TArray<int64>& InPads, const TArray<int64>& InStrides);
+
+	virtual ~FConvTransposeOperator();
+
+	FORCEINLINE virtual void SetWhetherWeightTensorWasFlipped(const bool bInIsWeightTensorFlipped) final;
+
+	virtual int32 GetNumberAuxiliaryTensors() const override final;
+
+	virtual bool ConfigureOutputAndInternalVariablesAndSanityChecks() override final;
+
+	virtual void ToGPU_RenderThread() override final;
+
+	virtual void ForwardCPU() override final;
+
+	virtual void ForwardGPU_RenderThread(FRDGBuilder* InOutGraphBuilder) override final;
+
+protected:
+	virtual bool SetAndConfigureStrides(const int32 InNumberConvolutionalDimensions) override final;
+
+	virtual bool SetAndConfigureAuxiliaryTensor() override final;
+
+private:
+	/**
+	 * True if an auxiliary tensor is needed, false otherwise.
+	 */
+	bool bIsXWithZerosNeeded;
+
+	/**
+	 * True if the weight tensor was already transposed.
+	 */
+	bool bIsWeightTensorFlipped;
+
+	/**
+	 * Zeros to add between each consecutive pixel of the input tensor.
+	 */
+	FNeuralInt64ArrayUInt32Buffer Zeros;
+	/**
+	 * Size of the input tensor X (size of X needed when moving it into XWithZeros).
+	 */
+	FNeuralInt64ArrayUInt32Buffer XSizesForGPU;
+};
+
+
+
+/* FConvBaseOperator inline functions
+ *****************************************************************************/
+
+void FConvTransposeOperator::SetWhetherWeightTensorWasFlipped(const bool bInIsWeightTensorFlipped)
+{
+	bIsWeightTensorFlipped = bInIsWeightTensorFlipped;
+}
diff --git a/Engine/Plugins/Experimental/NeuralNetworkInference/Source/NeuralNetworkInference/Internal/NeuralOperators/CosOperator.h b/Engine/Plugins/Experimental/NeuralNetworkInference/Source/NeuralNetworkInference/Internal/NeuralOperators/CosOperator.h
new file mode 100644
--- /dev/null
+++ b/Engine/Plugins/Experimental/NeuralNetworkInference/Source/NeuralNetworkInference/Internal/NeuralOperators/CosOperator.h
@@ -0,0 +1,26 @@
+// Copyright Epic Games, Inc. All Rights Reserved.
+
+#pragma once
+
+#include "CoreMinimal.h"
+#include "NeuralOperators/ElementWiseOperator.h"
+
+class NEURALNETWORKINFERENCE_API FCosOperator : public IElementWiseOperator
+{
+public:
+	FCosOperator(const bool bIsInlinedTensor);
+
+	virtual ~FCosOperator();
+
+	FORCEINLINE virtual void ForwardCPU() override final;
+};
+
+
+
+/* FCosOperator inline functions
+ *****************************************************************************/
+
+void FCosOperator::ForwardCPU()
+{
+	ForwardCPUWithFunction([](const float InValue) { return FMath::Cos(InValue); });
+}
diff --git a/Engine/Plugins/Experimental/NeuralNetworkInference/Source/NeuralNetworkInference/Internal/NeuralOperators/CoshOperator.h b/Engine/Plugins/Experimental/NeuralNetworkInference/Source/NeuralNetworkInference/Internal/NeuralOperators/CoshOperator.h
new file mode 100644
--- /dev/null
+++ b/Engine/Plugins/Experimental/NeuralNetworkInference/Source/NeuralNetworkInference/Internal/NeuralOperators/CoshOperator.h
@@ -0,0 +1,16 @@
+// Copyright Epic Games, Inc. All Rights Reserved.
+
+#pragma once
+
+#include "CoreMinimal.h"
+#include "NeuralOperators/ElementWiseOperator.h"
+
+class NEURALNETWORKINFERENCE_API FCoshOperator : public IElementWiseOperator
+{
+public:
+	FCoshOperator(const bool bIsInlinedTensor);
+
+	virtual ~FCoshOperator();
+
+	virtual void ForwardCPU() override final;
+};
diff --git a/Engine/Plugins/Experimental/NeuralNetworkInference/Source/NeuralNetworkInference/Internal/NeuralOperators/DivOperator.h b/Engine/Plugins/Experimental/NeuralNetworkInference/Source/NeuralNetworkInference/Internal/NeuralOperators/DivOperator.h
new file mode 100644
--- /dev/null
+++ b/Engine/Plugins/Experimental/NeuralNetworkInference/Source/NeuralNetworkInference/Internal/NeuralOperators/DivOperator.h
@@ -0,0 +1,26 @@
+// Copyright Epic Games, Inc. All Rights Reserved.
+
+#pragma once
+
+#include "CoreMinimal.h"
+#include "NeuralOperators/MultidirectionalBroadcastOperator.h"
+
+class NEURALNETWORKINFERENCE_API FDivOperator : public IMultidirectionalBroadcastOperator
+{
+public:
+	FDivOperator(const TSet<uint32>& InPotentialInlinedTensors);
+
+	virtual ~FDivOperator();
+
+	FORCEINLINE virtual void ForwardCPU() override final;
+};
+
+
+
+/* FDivOperator inline functions
+ *****************************************************************************/
+
+void FDivOperator::ForwardCPU()
+{
+	ForwardCPUWithFunction([](const float InValueA, const float InValueB) { return InValueA / InValueB; });
+}
diff --git a/Engine/Plugins/Experimental/NeuralNetworkInference/Source/NeuralNetworkInference/Internal/NeuralOperators/ElementWiseOperator.h b/Engine/Plugins/Experimental/NeuralNetworkInference/Source/NeuralNetworkInference/Internal/NeuralOperators/ElementWiseOperator.h
new file mode 100644
--- /dev/null
+++ b/Engine/Plugins/Experimental/NeuralNetworkInference/Source/NeuralNetworkInference/Internal/NeuralOperators/ElementWiseOperator.h
@@ -0,0 +1,37 @@
+// Copyright Epic Games, Inc. All Rights Reserved.
+
+#pragma once
+
+#include "CoreMinimal.h"
+#include "NeuralOperator.h"
+#include "NeuralOperatorEnumClasses.h"
+#include "NeuralTensor.h"
+
+class NEURALNETWORKINFERENCE_API IElementWiseOperator : public FNeuralOperator
+{
+public:
+	IElementWiseOperator(const FString& InName, const int32 InVersion, const EElementWiseOperator InElementWiseOperator, const bool bIsInlinedTensor,
+		const TArray<float>& InAttributes = TArray<float>());
+
+	virtual ~IElementWiseOperator();
+
+	virtual bool ConfigureOutputAndInternalVariablesAndSanityChecks() override final;
+
+	virtual void ForwardGPU_RenderThread(FRDGBuilder* InOutGraphBuilder) override final;
+
+protected:
+	TArray<float> Attributes;
+
+	/**
+	 * This is the function that child classes must call on ForwardCPU().
+	 */
+	virtual void ForwardCPUWithFunction(float InOperatorFunction(const float)) final;
+
+	/**
+	 * This is the function that child classes must call on ForwardCPU().
+	 */
+	virtual void ForwardCPUWithFunction(float InOperatorFunction(const float, const float)) final;
+
+private:
+	const EElementWiseOperator ElementWiseOperator;
+};
diff --git a/Engine/Plugins/Experimental/NeuralNetworkInference/Source/NeuralNetworkInference/Internal/NeuralOperators/ExpOperator.h b/Engine/Plugins/Experimental/NeuralNetworkInference/Source/NeuralNetworkInference/Internal/NeuralOperators/ExpOperator.h
new file mode 100644
--- /dev/null
+++ b/Engine/Plugins/Experimental/NeuralNetworkInference/Source/NeuralNetworkInference/Internal/NeuralOperators/ExpOperator.h
@@ -0,0 +1,26 @@
+// Copyright Epic Games, Inc. All Rights Reserved.
+
+#pragma once
+
+#include "CoreMinimal.h"
+#include "NeuralOperators/ElementWiseOperator.h"
+
+class NEURALNETWORKINFERENCE_API FExpOperator : public IElementWiseOperator
+{
+public:
+	FExpOperator(const bool bIsInlinedTensor);
+
+	virtual ~FExpOperator();
+
+	FORCEINLINE virtual void ForwardCPU() override final;
+};
+
+
+
+/* FExpOperator inline functions
+ *****************************************************************************/
+
+void FExpOperator::ForwardCPU()
+{
+	ForwardCPUWithFunction([](const float InValue) { return FMath::Exp(InValue); });
+}
diff --git a/Engine/Plugins/Experimental/NeuralNetworkInference/Source/NeuralNetworkInference/Internal/NeuralOperators/FloorOperator.h b/Engine/Plugins/Experimental/NeuralNetworkInference/Source/NeuralNetworkInference/Internal/NeuralOperators/FloorOperator.h
new file mode 100644
--- /dev/null
+++ b/Engine/Plugins/Experimental/NeuralNetworkInference/Source/NeuralNetworkInference/Internal/NeuralOperators/FloorOperator.h
@@ -0,0 +1,26 @@
+// Copyright Epic Games, Inc. All Rights Reserved.
+
+#pragma once
+
+#include "CoreMinimal.h"
+#include "NeuralOperators/ElementWiseOperator.h"
+
+class NEURALNETWORKINFERENCE_API FFloorOperator : public IElementWiseOperator
+{
+public:
+	FFloorOperator(const bool bIsInlinedTensor);
+
+	virtual ~FFloorOperator();
+
+	FORCEINLINE virtual void ForwardCPU() override final;
+};
+
+
+
+/* FFloorOperator inline functions
+ *****************************************************************************/
+
+void FFloorOperator::ForwardCPU()
+{
+	ForwardCPUWithFunction([](const float InValue) { return FMath::Floor(InValue); });
+}
diff --git a/Engine/Plugins/Experimental/NeuralNetworkInference/Source/NeuralNetworkInference/Internal/NeuralOperators/GemmOperator.h b/Engine/Plugins/Experimental/NeuralNetworkInference/Source/NeuralNetworkInference/Internal/NeuralOperators/GemmOperator.h
new file mode 100644
--- /dev/null
+++ b/Engine/Plugins/Experimental/NeuralNetworkInference/Source/NeuralNetworkInference/Internal/NeuralOperators/GemmOperator.h
@@ -0,0 +1,46 @@
+// Copyright Epic Games, Inc. All Rights Reserved.
+
+#pragma once
+
+#include "CoreMinimal.h"
+#include "ModelProto.h"
+#include "NeuralOperator.h"
+
+class NEURALNETWORKINFERENCE_API FGemmOperator : public FNeuralOperator
+{
+public:
+	FGemmOperator(const FNodeProto& InNodeProto);
+
+	FGemmOperator(const float InAlpha = 1.f, const float InBeta = 1.f, const bool bInTransA = false, const bool bInTransB = false);
+
+	virtual ~FGemmOperator();
+
+	virtual bool ConfigureOutputAndInternalVariablesAndSanityChecks() override final;
+
+	virtual void ForwardCPU() override final;
+
+	virtual void ForwardGPU_RenderThread(FRDGBuilder* InOutGraphBuilder) override final;
+
+private:
+	float Alpha;
+	float Beta;
+	bool bTransA;
+	bool bTransB;
+	/**
+	 * GEMM-required params
+	 */
+	uint32 OutputRows;
+	uint32 OutputColumns;
+	uint32 AColsOrBRows;
+	uint32 AStrideX;
+	uint32 AStrideY;
+	uint32 BStrideX;
+	uint32 BStrideY;
+	uint32 OutputStride;
+	/**
+	 * C is unidirectional broadcastable to the output. We need this to account for edge cases.
+	 * In the usual case (output size == C size), CSize = OutputSize2D
+	 */
+	uint32 CSizeX;
+	uint32 CSizeY;
+};
diff --git a/Engine/Plugins/Experimental/NeuralNetworkInference/Source/NeuralNetworkInference/Internal/NeuralOperators/LeakyReluOperator.h b/Engine/Plugins/Experimental/NeuralNetworkInference/Source/NeuralNetworkInference/Internal/NeuralOperators/LeakyReluOperator.h
new file mode 100644
--- /dev/null
+++ b/Engine/Plugins/Experimental/NeuralNetworkInference/Source/NeuralNetworkInference/Internal/NeuralOperators/LeakyReluOperator.h
@@ -0,0 +1,29 @@
+// Copyright Epic Games, Inc. All Rights Reserved.
+
+#pragma once
+
+#include "CoreMinimal.h"
+#include "ModelProto.h"
+#include "NeuralOperators/ElementWiseOperator.h"
+
+class NEURALNETWORKINFERENCE_API FLeakyReluOperator : public IElementWiseOperator
+{
+public:
+	FLeakyReluOperator(const bool bIsInlinedTensor, const FNodeProto& InNodeProto);
+
+	FLeakyReluOperator(const bool bIsInlinedTensor, const float InAlpha = 0.01f);
+
+	virtual ~FLeakyReluOperator();
+
+	FORCEINLINE virtual void ForwardCPU() override final;
+};
+
+
+
+/* FLeakyReluOperator inline functions
+ *****************************************************************************/
+
+void FLeakyReluOperator::ForwardCPU()
+{
+	ForwardCPUWithFunction([](const float InValue, const float InAttribute) { return (InValue < 0 ? InAttribute * InValue : InValue); });
+}
diff --git a/Engine/Plugins/Experimental/NeuralNetworkInference/Source/NeuralNetworkInference/Internal/NeuralOperators/LogOperator.h b/Engine/Plugins/Experimental/NeuralNetworkInference/Source/NeuralNetworkInference/Internal/NeuralOperators/LogOperator.h
new file mode 100644
--- /dev/null
+++ b/Engine/Plugins/Experimental/NeuralNetworkInference/Source/NeuralNetworkInference/Internal/NeuralOperators/LogOperator.h
@@ -0,0 +1,26 @@
+// Copyright Epic Games, Inc. All Rights Reserved.
+
+#pragma once
+
+#include "CoreMinimal.h"
+#include "NeuralOperators/ElementWiseOperator.h"
+
+class NEURALNETWORKINFERENCE_API FLogOperator : public IElementWiseOperator
+{
+public:
+	FLogOperator(const bool bIsInlinedTensor);
+
+	virtual ~FLogOperator();
+
+	FORCEINLINE virtual void ForwardCPU() override final;
+};
+
+
+
+/* FLogOperator inline functions
+ *****************************************************************************/
+
+void FLogOperator::ForwardCPU()
+{
+	ForwardCPUWithFunction([](const float InValue) { return FMath::Loge(InValue); });
+}
diff --git a/Engine/Plugins/Experimental/NeuralNetworkInference/Source/NeuralNetworkInference/Internal/NeuralOperators/MulOperator.h b/Engine/Plugins/Experimental/NeuralNetworkInference/Source/NeuralNetworkInference/Internal/NeuralOperators/MulOperator.h
new file mode 100644
--- /dev/null
+++ b/Engine/Plugins/Experimental/NeuralNetworkInference/Source/NeuralNetworkInference/Internal/NeuralOperators/MulOperator.h
@@ -0,0 +1,26 @@
+// Copyright Epic Games, Inc. All Rights Reserved.
+
+#pragma once
+
+#include "CoreMinimal.h"
+#include "NeuralOperators/MultidirectionalBroadcastOperator.h"
+
+class NEURALNETWORKINFERENCE_API FMulOperator : public IMultidirectionalBroadcastOperator
+{
+public:
+	FMulOperator(const TSet<uint32>& InPotentialInlinedTensors);
+
+	virtual ~FMulOperator();
+
+	FORCEINLINE virtual void ForwardCPU() override final;
+};
+
+
+
+/* FMulOperator inline functions
+ *****************************************************************************/
+
+void FMulOperator::ForwardCPU()
+{
+	ForwardCPUWithFunction([](const float InValueA, const float InValueB) { return InValueA * InValueB; });
+}
diff --git a/Engine/Plugins/Experimental/NeuralNetworkInference/Source/NeuralNetworkInference/Internal/NeuralOperators/MultidirectionalBroadcastOperator.h b/Engine/Plugins/Experimental/NeuralNetworkInference/Source/NeuralNetworkInference/Internal/NeuralOperators/MultidirectionalBroadcastOperator.h
new file mode 100644
--- /dev/null
+++ b/Engine/Plugins/Experimental/NeuralNetworkInference/Source/NeuralNetworkInference/Internal/NeuralOperators/MultidirectionalBroadcastOperator.h
@@ -0,0 +1,46 @@
+// Copyright Epic Games, Inc. All Rights Reserved.
+
+#pragma once
+
+#include "CoreMinimal.h"
+#include "NeuralInt64ArrayUInt32Buffer.h"
+#include "NeuralOperator.h"
+#include "NeuralOperatorEnumClasses.h"
+
+struct FReadBuffer;
+
+class NEURALNETWORKINFERENCE_API IMultidirectionalBroadcastOperator : public FNeuralOperator
+{
+public:
+	/**
+	 * Not inlined layer if an empty InPotentialInlinedTensors is given.
+	 * It will try to inline the input tensors otherwise, where only the input tensors with the given indexes (InPotentialInlinedTensors) can be potentially inlined layers.
+	 */
+	IMultidirectionalBroadcastOperator(const FString& InName, const int32 InVersion, const EMultidirectionalBroadcastOperator InMultidirectionalBroadcastOperator,
+		const TSet<uint32>& InPotentialInlinedTensors);
+
+	virtual ~IMultidirectionalBroadcastOperator();
+
+	virtual bool ConfigureOutputAndInternalVariablesAndSanityChecks() override final;
+
+	virtual void ToGPU_RenderThread() override final;
+
+	virtual void ForwardGPU_RenderThread(FRDGBuilder* InOutGraphBuilder) override final;
+
+protected:
+	virtual bool EstimateInlinedTensorFromPotentialOnes() override final;
+
+	virtual void ShapesToGPU() final;
+
+	/**
+	 * This is the function that child classes must call on ForwardCPU().
+	 */
+	virtual void ForwardCPUWithFunction(float InOperatorFunction(const float, const float)) final;
+
+private:
+	const EMultidirectionalBroadcastOperator MultidirectionalBroadcastOperator;
+	EMultidirectionalBroadcastShapeMode MultidirectionalBroadcastShapeMode;
+	FNeuralInt64ArrayUInt32Buffer ShapesX;
+	FNeuralInt64ArrayUInt32Buffer ShapesY;
+	FNeuralInt64ArrayUInt32Buffer ShapesOutput;
+};
diff --git a/Engine/Plugins/Experimental/NeuralNetworkInference/Source/NeuralNetworkInference/Internal/NeuralOperators/NegOperator.h b/Engine/Plugins/Experimental/NeuralNetworkInference/Source/NeuralNetworkInference/Internal/NeuralOperators/NegOperator.h
new file mode 100644
--- /dev/null
+++ b/Engine/Plugins/Experimental/NeuralNetworkInference/Source/NeuralNetworkInference/Internal/NeuralOperators/NegOperator.h
@@ -0,0 +1,26 @@
+// Copyright Epic Games, Inc. All Rights Reserved.
+
+#pragma once
+
+#include "CoreMinimal.h"
+#include "NeuralOperators/ElementWiseOperator.h"
+
+class NEURALNETWORKINFERENCE_API FNegOperator : public IElementWiseOperator
+{
+public:
+	FNegOperator(const bool bIsInlinedTensor);
+
+	virtual ~FNegOperator();
+
+	FORCEINLINE virtual void ForwardCPU() override final;
+};
+
+
+
+/* FNegOperator inline functions
+ *****************************************************************************/
+
+void FNegOperator::ForwardCPU()
+{
+	ForwardCPUWithFunction([](const float InValue) { return -InValue; });
+}
diff --git a/Engine/Plugins/Experimental/NeuralNetworkInference/Source/NeuralNetworkInference/Internal/NeuralOperators/PowOperator.h b/Engine/Plugins/Experimental/NeuralNetworkInference/Source/NeuralNetworkInference/Internal/NeuralOperators/PowOperator.h
new file mode 100644
--- /dev/null
+++ b/Engine/Plugins/Experimental/NeuralNetworkInference/Source/NeuralNetworkInference/Internal/NeuralOperators/PowOperator.h
@@ -0,0 +1,26 @@
+// Copyright Epic Games, Inc. All Rights Reserved.
+
+#pragma once
+
+#include "CoreMinimal.h"
+#include "NeuralOperators/MultidirectionalBroadcastOperator.h"
+
+class NEURALNETWORKINFERENCE_API FPowOperator : public IMultidirectionalBroadcastOperator
+{
+public:
+	FPowOperator(const TSet<uint32>& InPotentialInlinedTensors);
+
+	virtual ~FPowOperator();
+
+	FORCEINLINE virtual void ForwardCPU() override final;
+};
+
+
+
+/* FPowOperator inline functions
+ *****************************************************************************/
+
+void FPowOperator::ForwardCPU()
+{
+	ForwardCPUWithFunction(FMath::Pow);
+}
diff --git a/Engine/Plugins/Experimental/NeuralNetworkInference/Source/NeuralNetworkInference/Internal/NeuralOperators/ReciprocalOperator.h b/Engine/Plugins/Experimental/NeuralNetworkInference/Source/NeuralNetworkInference/Internal/NeuralOperators/ReciprocalOperator.h
new file mode 100644
--- /dev/null
+++ b/Engine/Plugins/Experimental/NeuralNetworkInference/Source/NeuralNetworkInference/Internal/NeuralOperators/ReciprocalOperator.h
@@ -0,0 +1,26 @@
+// Copyright Epic Games, Inc. All Rights Reserved.
+
+#pragma once
+
+#include "CoreMinimal.h"
+#include "NeuralOperators/ElementWiseOperator.h"
+
+class NEURALNETWORKINFERENCE_API FReciprocalOperator : public IElementWiseOperator
+{
+public:
+	FReciprocalOperator(const bool bIsInlinedTensor);
+
+	virtual ~FReciprocalOperator();
+
+	FORCEINLINE virtual void ForwardCPU() override final;
+};
+
+
+
+/* FReciprocalOperator inline functions
+ *****************************************************************************/
+
+void FReciprocalOperator::ForwardCPU()
+{
+	ForwardCPUWithFunction([](const float InValue) { return 1.f/InValue; });
+}
diff --git a/Engine/Plugins/Experimental/NeuralNetworkInference/Source/NeuralNetworkInference/Internal/NeuralOperators/ReluOperator.h b/Engine/Plugins/Experimental/NeuralNetworkInference/Source/NeuralNetworkInference/Internal/NeuralOperators/ReluOperator.h
new file mode 100644
--- /dev/null
+++ b/Engine/Plugins/Experimental/NeuralNetworkInference/Source/NeuralNetworkInference/Internal/NeuralOperators/ReluOperator.h
@@ -0,0 +1,26 @@
+// Copyright Epic Games, Inc. All Rights Reserved.
+
+#pragma once
+
+#include "CoreMinimal.h"
+#include "NeuralOperators/ElementWiseOperator.h"
+
+class NEURALNETWORKINFERENCE_API FReluOperator : public IElementWiseOperator
+{
+public:
+	FReluOperator(const bool bIsInlinedTensor);
+
+	virtual ~FReluOperator();
+
+	FORCEINLINE virtual void ForwardCPU() override final;
+};
+
+
+
+/* FReluOperator inline functions
+ *****************************************************************************/
+
+void FReluOperator::ForwardCPU()
+{
+	ForwardCPUWithFunction([](const float InValue) { return (InValue > 0.f ? InValue : 0.f); });
+}
diff --git a/Engine/Plugins/Experimental/NeuralNetworkInference/Source/NeuralNetworkInference/Internal/NeuralOperators/ReshapeOperator.h b/Engine/Plugins/Experimental/NeuralNetworkInference/Source/NeuralNetworkInference/Internal/NeuralOperators/ReshapeOperator.h
new file mode 100644
--- /dev/null
+++ b/Engine/Plugins/Experimental/NeuralNetworkInference/Source/NeuralNetworkInference/Internal/NeuralOperators/ReshapeOperator.h
@@ -0,0 +1,68 @@
+// Copyright Epic Games, Inc. All Rights Reserved.
+
+#pragma once
+
+#include "CoreMinimal.h"
+#include "ModelProto.h"
+#include "NeuralOperator.h"
+
+/**
+ *  0 = Copy from input
+ * -1 = Infer from remaining dimensions
+ */
+class NEURALNETWORKINFERENCE_API FReshapeOperator : public FNeuralOperator
+{
+public:
+	FReshapeOperator(const bool bIsInlinedTensor, const FNodeProto& InNodeProto);
+
+	FReshapeOperator(const bool bIsInlinedTensor, const int64 InAllowZero = 0);
+
+	virtual ~FReshapeOperator();
+
+	virtual bool ConfigureOutputAndInternalVariablesAndSanityChecks() override final;
+
+	virtual bool SetWhetherInputTensorsNeedTransferToGPUForGPUMode() override final;
+
+	virtual void ForwardCPU() override final;
+
+	virtual void ForwardGPU_RenderThread(FRDGBuilder* InOutGraphBuilder) override final;
+
+	FORCEINLINE virtual bool NeedsPostForward() const override final;
+
+	virtual void PostForwardCPU() override final;
+
+	FORCEINLINE virtual void PostForwardGPU_RenderThread(FRDGBuilder* InOutGraphBuilder) override final;
+
+private:
+	/**
+	 * (Optional) By default, when any value in the 'shape' input is equal to zero the corresponding dimension value is copied from the input tensor dynamically.
+	 * AllowZero = 1 indicates that if any value in the 'shape' input is set to zero, the zero value is honored, similar to NumPy.
+	 */
+	int64 AllowZero;
+	
+	/**
+	 * If not inlined operator, this vector will be empty. If inlined, it will represent the original input size that PostForwardCPU/GPU will set.
+	 */
+	TArray<int64> InputSizeIfInlined;
+
+	/**
+	 * Note: It is assumed that Minus1Index == -1 up to 1 times, otherwise, this will fail.
+	 */
+	virtual TArray<int64> Remove0sAndMinus1sFromShape(const FNeuralTensor& InShapeTensor, const FNeuralTensor& InTensor) const final;
+};
+
+
+
+/* FNeuralOperator inline functions
+ *****************************************************************************/
+
+bool FReshapeOperator::NeedsPostForward() const
+{
+	return InlinedTensor > -1;
+}
+
+void FReshapeOperator::PostForwardGPU_RenderThread(FRDGBuilder* InOutGraphBuilder)
+{
+	// It only needs to reshape input size
+	PostForwardCPU();
+}
diff --git a/Engine/Plugins/Experimental/NeuralNetworkInference/Source/NeuralNetworkInference/Internal/NeuralOperators/RoundOperator.h b/Engine/Plugins/Experimental/NeuralNetworkInference/Source/NeuralNetworkInference/Internal/NeuralOperators/RoundOperator.h
new file mode 100644
--- /dev/null
+++ b/Engine/Plugins/Experimental/NeuralNetworkInference/Source/NeuralNetworkInference/Internal/NeuralOperators/RoundOperator.h
@@ -0,0 +1,26 @@
+// Copyright Epic Games, Inc. All Rights Reserved.
+
+#pragma once
+
+#include "CoreMinimal.h"
+#include "NeuralOperators/ElementWiseOperator.h"
+
+class NEURALNETWORKINFERENCE_API FRoundOperator : public IElementWiseOperator
+{
+public:
+	FRoundOperator(const bool bIsInlinedTensor);
+
+	virtual ~FRoundOperator();
+
+	FORCEINLINE virtual void ForwardCPU() override final;
+};
+
+
+
+/* FRoundOperator inline functions
+ *****************************************************************************/
+
+void FRoundOperator::ForwardCPU()
+{
+	ForwardCPUWithFunction([](const float InValue) { return FMath::RoundToFloat(InValue); });
+}
diff --git a/Engine/Plugins/Experimental/NeuralNetworkInference/Source/NeuralNetworkInference/Internal/NeuralOperators/SigmoidOperator.h b/Engine/Plugins/Experimental/NeuralNetworkInference/Source/NeuralNetworkInference/Internal/NeuralOperators/SigmoidOperator.h
new file mode 100644
--- /dev/null
+++ b/Engine/Plugins/Experimental/NeuralNetworkInference/Source/NeuralNetworkInference/Internal/NeuralOperators/SigmoidOperator.h
@@ -0,0 +1,26 @@
+// Copyright Epic Games, Inc. All Rights Reserved.
+
+#pragma once
+
+#include "CoreMinimal.h"
+#include "NeuralOperators/ElementWiseOperator.h"
+
+class NEURALNETWORKINFERENCE_API FSigmoidOperator : public IElementWiseOperator
+{
+public:
+	FSigmoidOperator(const bool bIsInlinedTensor);
+
+	virtual ~FSigmoidOperator();
+
+	FORCEINLINE virtual void ForwardCPU() override final;
+};
+
+
+
+/* FSigmoidOperator inline functions
+ *****************************************************************************/
+
+void FSigmoidOperator::ForwardCPU()
+{
+	ForwardCPUWithFunction([](const float InValue) { return 1.f/(1.f + FMath::Exp(-InValue)); });
+}
diff --git a/Engine/Plugins/Experimental/NeuralNetworkInference/Source/NeuralNetworkInference/Internal/NeuralOperators/SignOperator.h b/Engine/Plugins/Experimental/NeuralNetworkInference/Source/NeuralNetworkInference/Internal/NeuralOperators/SignOperator.h
new file mode 100644
--- /dev/null
+++ b/Engine/Plugins/Experimental/NeuralNetworkInference/Source/NeuralNetworkInference/Internal/NeuralOperators/SignOperator.h
@@ -0,0 +1,26 @@
+// Copyright Epic Games, Inc. All Rights Reserved.
+
+#pragma once
+
+#include "CoreMinimal.h"
+#include "NeuralOperators/ElementWiseOperator.h"
+
+class NEURALNETWORKINFERENCE_API FSignOperator : public IElementWiseOperator
+{
+public:
+	FSignOperator(const bool bIsInlinedTensor);
+
+	virtual ~FSignOperator();
+
+	FORCEINLINE virtual void ForwardCPU() override final;
+};
+
+
+
+/* FSignOperator inline functions
+ *****************************************************************************/
+
+void FSignOperator::ForwardCPU()
+{
+	ForwardCPUWithFunction([](const float InValue) { return FMath::Sign(InValue); });
+}
diff --git a/Engine/Plugins/Experimental/NeuralNetworkInference/Source/NeuralNetworkInference/Internal/NeuralOperators/SinOperator.h b/Engine/Plugins/Experimental/NeuralNetworkInference/Source/NeuralNetworkInference/Internal/NeuralOperators/SinOperator.h
new file mode 100644
--- /dev/null
+++ b/Engine/Plugins/Experimental/NeuralNetworkInference/Source/NeuralNetworkInference/Internal/NeuralOperators/SinOperator.h
@@ -0,0 +1,26 @@
+// Copyright Epic Games, Inc. All Rights Reserved.
+
+#pragma once
+
+#include "CoreMinimal.h"
+#include "NeuralOperators/ElementWiseOperator.h"
+
+class NEURALNETWORKINFERENCE_API FSinOperator : public IElementWiseOperator
+{
+public:
+	FSinOperator(const bool bIsInlinedTensor);
+
+	virtual ~FSinOperator();
+
+	FORCEINLINE virtual void ForwardCPU() override final;
+};
+
+
+
+/* FSinOperator inline functions
+ *****************************************************************************/
+
+void FSinOperator::ForwardCPU()
+{
+	ForwardCPUWithFunction([](const float InValue) { return FMath::Sin(InValue); });
+}
diff --git a/Engine/Plugins/Experimental/NeuralNetworkInference/Source/NeuralNetworkInference/Internal/NeuralOperators/SinhOperator.h b/Engine/Plugins/Experimental/NeuralNetworkInference/Source/NeuralNetworkInference/Internal/NeuralOperators/SinhOperator.h
new file mode 100644
--- /dev/null
+++ b/Engine/Plugins/Experimental/NeuralNetworkInference/Source/NeuralNetworkInference/Internal/NeuralOperators/SinhOperator.h
@@ -0,0 +1,26 @@
+// Copyright Epic Games, Inc. All Rights Reserved.
+
+#pragma once
+
+#include "CoreMinimal.h"
+#include "NeuralOperators/ElementWiseOperator.h"
+
+class NEURALNETWORKINFERENCE_API FSinhOperator : public IElementWiseOperator
+{
+public:
+	FSinhOperator(const bool bIsInlinedTensor);
+
+	virtual ~FSinhOperator();
+
+	FORCEINLINE virtual void ForwardCPU() override final;
+};
+
+
+
+/* FSinhOperator inline functions
+ *****************************************************************************/
+
+void FSinhOperator::ForwardCPU()
+{
+	ForwardCPUWithFunction([](const float InValue) { return FMath::Sinh(InValue); });
+}
diff --git a/Engine/Plugins/Experimental/NeuralNetworkInference/Source/NeuralNetworkInference/Internal/NeuralOperators/SqrtOperator.h b/Engine/Plugins/Experimental/NeuralNetworkInference/Source/NeuralNetworkInference/Internal/NeuralOperators/SqrtOperator.h
new file mode 100644
--- /dev/null
+++ b/Engine/Plugins/Experimental/NeuralNetworkInference/Source/NeuralNetworkInference/Internal/NeuralOperators/SqrtOperator.h
@@ -0,0 +1,26 @@
+// Copyright Epic Games, Inc. All Rights Reserved.
+
+#pragma once
+
+#include "CoreMinimal.h"
+#include "NeuralOperators/ElementWiseOperator.h"
+
+class NEURALNETWORKINFERENCE_API FSqrtOperator : public IElementWiseOperator
+{
+public:
+	FSqrtOperator(const bool bIsInlinedTensor);
+
+	virtual ~FSqrtOperator();
+
+	FORCEINLINE virtual void ForwardCPU() override final;
+};
+
+
+
+/* FSqrtOperator inline functions
+ *****************************************************************************/
+
+void FSqrtOperator::ForwardCPU()
+{
+	ForwardCPUWithFunction([](const float InValue) { return FMath::Sqrt(InValue); });
+}
diff --git a/Engine/Plugins/Experimental/NeuralNetworkInference/Source/NeuralNetworkInference/Internal/NeuralOperators/SqueezeOperator.h b/Engine/Plugins/Experimental/NeuralNetworkInference/Source/NeuralNetworkInference/Internal/NeuralOperators/SqueezeOperator.h
new file mode 100644
--- /dev/null
+++ b/Engine/Plugins/Experimental/NeuralNetworkInference/Source/NeuralNetworkInference/Internal/NeuralOperators/SqueezeOperator.h
@@ -0,0 +1,62 @@
+// Copyright Epic Games, Inc. All Rights Reserved.
+
+#pragma once
+
+#include "CoreMinimal.h"
+#include "ModelProto.h"
+#include "NeuralOperator.h"
+
+
+class NEURALNETWORKINFERENCE_API FSqueezeOperator : public FNeuralOperator
+{
+public:
+	FSqueezeOperator(const bool bIsInlinedNeuralTensor);
+
+	virtual ~FSqueezeOperator();
+
+	virtual bool ConfigureOutputAndInternalVariablesAndSanityChecks() override final;
+
+	virtual bool SetWhetherInputTensorsNeedTransferToGPUForGPUMode() override final;
+
+	virtual void ForwardCPU() override final;
+
+	virtual void ForwardGPU_RenderThread(FRDGBuilder* InOutGraphBuilder) override final;
+
+	FORCEINLINE virtual bool NeedsPostForward() const override final;
+
+	virtual void PostForwardCPU() override final;
+
+	FORCEINLINE virtual void PostForwardGPU_RenderThread(FRDGBuilder* InOutGraphBuilder) override final;
+
+private:
+	/**
+	 * If not inlined operator, this vector will be empty. If inlined, it will represent the original input size that PostForwardCPU/GPU will set.
+	 */
+	TArray<int64> InputSizeIfInlined;
+
+	/**
+	 * It squeezes (i.e., removes the dimensions of size 1) of the tensor dimensions specified by InAxesTensor.
+	 */
+	virtual TArray<int64> Remove1sFromShape(const FNeuralTensor& InTensor, const FNeuralTensor& InAxesTensor) const final;
+
+	/**
+	 * It squeezes (i.e., removes the dimensions of size 1) of all the tensor.
+	 */
+	virtual TArray<int64> Remove1sFromShape(const FNeuralTensor& InTensor) const final;
+};
+
+
+
+/* FNeuralOperator inline functions
+ *****************************************************************************/
+
+bool FSqueezeOperator::NeedsPostForward() const
+{
+	return InlinedTensor > -1;
+}
+
+void FSqueezeOperator::PostForwardGPU_RenderThread(FRDGBuilder* InOutGraphBuilder)
+{
+	// It only needs to reshape input size
+	PostForwardCPU();
+}
diff --git a/Engine/Plugins/Experimental/NeuralNetworkInference/Source/NeuralNetworkInference/Internal/NeuralOperators/SubOperator.h b/Engine/Plugins/Experimental/NeuralNetworkInference/Source/NeuralNetworkInference/Internal/NeuralOperators/SubOperator.h
new file mode 100644
--- /dev/null
+++ b/Engine/Plugins/Experimental/NeuralNetworkInference/Source/NeuralNetworkInference/Internal/NeuralOperators/SubOperator.h
@@ -0,0 +1,26 @@
+// Copyright Epic Games, Inc. All Rights Reserved.
+
+#pragma once
+
+#include "CoreMinimal.h"
+#include "NeuralOperators/MultidirectionalBroadcastOperator.h"
+
+class NEURALNETWORKINFERENCE_API FSubOperator : public IMultidirectionalBroadcastOperator
+{
+public:
+	FSubOperator(const TSet<uint32>& InPotentialInlinedTensors);
+
+	virtual ~FSubOperator();
+
+	FORCEINLINE virtual void ForwardCPU() override final;
+};
+
+
+
+/* FSubOperator inline functions
+ *****************************************************************************/
+
+void FSubOperator::ForwardCPU()
+{
+	ForwardCPUWithFunction([](const float InValueA, const float InValueB) { return InValueA - InValueB; });
+}
diff --git a/Engine/Plugins/Experimental/NeuralNetworkInference/Source/NeuralNetworkInference/Internal/NeuralOperators/TanOperator.h b/Engine/Plugins/Experimental/NeuralNetworkInference/Source/NeuralNetworkInference/Internal/NeuralOperators/TanOperator.h
new file mode 100644
--- /dev/null
+++ b/Engine/Plugins/Experimental/NeuralNetworkInference/Source/NeuralNetworkInference/Internal/NeuralOperators/TanOperator.h
@@ -0,0 +1,26 @@
+// Copyright Epic Games, Inc. All Rights Reserved.
+
+#pragma once
+
+#include "CoreMinimal.h"
+#include "NeuralOperators/ElementWiseOperator.h"
+
+class NEURALNETWORKINFERENCE_API FTanOperator : public IElementWiseOperator
+{
+public:
+	FTanOperator(const bool bIsInlinedTensor);
+
+	virtual ~FTanOperator();
+
+	FORCEINLINE virtual void ForwardCPU() override final;
+};
+
+
+
+/* FTanOperator inline functions
+ *****************************************************************************/
+
+void FTanOperator::ForwardCPU()
+{
+	ForwardCPUWithFunction([](const float InValue) { return FMath::Tan(InValue); });
+}
diff --git a/Engine/Plugins/Experimental/NeuralNetworkInference/Source/NeuralNetworkInference/Internal/NeuralOperators/TanhOperator.h b/Engine/Plugins/Experimental/NeuralNetworkInference/Source/NeuralNetworkInference/Internal/NeuralOperators/TanhOperator.h
new file mode 100644
--- /dev/null
+++ b/Engine/Plugins/Experimental/NeuralNetworkInference/Source/NeuralNetworkInference/Internal/NeuralOperators/TanhOperator.h
@@ -0,0 +1,16 @@
+// Copyright Epic Games, Inc. All Rights Reserved.
+
+#pragma once
+
+#include "CoreMinimal.h"
+#include "NeuralOperators/ElementWiseOperator.h"
+
+class NEURALNETWORKINFERENCE_API FTanhOperator : public IElementWiseOperator
+{
+public:
+	FTanhOperator(const bool bIsInlinedTensor);
+
+	virtual ~FTanhOperator();
+
+	virtual void ForwardCPU() override final;
+};
diff --git a/Engine/Plugins/Experimental/NeuralNetworkInference/Source/NeuralNetworkInference/Public/NeuralInt64ArrayUInt32Buffer.h b/Engine/Plugins/Experimental/NeuralNetworkInference/Source/NeuralNetworkInference/Public/NeuralInt64ArrayUInt32Buffer.h
deleted file mode 100644
diff --git a/Engine/Plugins/Experimental/NeuralNetworkInference/Source/NeuralNetworkInference/Public/NeuralOperator.h b/Engine/Plugins/Experimental/NeuralNetworkInference/Source/NeuralNetworkInference/Public/NeuralOperator.h
deleted file mode 100644
diff --git a/Engine/Plugins/Experimental/NeuralNetworkInference/Source/NeuralNetworkInference/Public/NeuralOperators.h b/Engine/Plugins/Experimental/NeuralNetworkInference/Source/NeuralNetworkInference/Public/NeuralOperators.h
deleted file mode 100644
diff --git a/Engine/Plugins/Experimental/NeuralNetworkInference/Source/NeuralNetworkInference/Public/NeuralOperators/AbsOperator.h b/Engine/Plugins/Experimental/NeuralNetworkInference/Source/NeuralNetworkInference/Public/NeuralOperators/AbsOperator.h
deleted file mode 100644
diff --git a/Engine/Plugins/Experimental/NeuralNetworkInference/Source/NeuralNetworkInference/Public/NeuralOperators/AcosOperator.h b/Engine/Plugins/Experimental/NeuralNetworkInference/Source/NeuralNetworkInference/Public/NeuralOperators/AcosOperator.h
deleted file mode 100644
diff --git a/Engine/Plugins/Experimental/NeuralNetworkInference/Source/NeuralNetworkInference/Public/NeuralOperators/AddOperator.h b/Engine/Plugins/Experimental/NeuralNetworkInference/Source/NeuralNetworkInference/Public/NeuralOperators/AddOperator.h
deleted file mode 100644
diff --git a/Engine/Plugins/Experimental/NeuralNetworkInference/Source/NeuralNetworkInference/Public/NeuralOperators/AsinOperator.h b/Engine/Plugins/Experimental/NeuralNetworkInference/Source/NeuralNetworkInference/Public/NeuralOperators/AsinOperator.h
deleted file mode 100644
diff --git a/Engine/Plugins/Experimental/NeuralNetworkInference/Source/NeuralNetworkInference/Public/NeuralOperators/AtanOperator.h b/Engine/Plugins/Experimental/NeuralNetworkInference/Source/NeuralNetworkInference/Public/NeuralOperators/AtanOperator.h
deleted file mode 100644
diff --git a/Engine/Plugins/Experimental/NeuralNetworkInference/Source/NeuralNetworkInference/Public/NeuralOperators/BatchNormalizationOperator.h b/Engine/Plugins/Experimental/NeuralNetworkInference/Source/NeuralNetworkInference/Public/NeuralOperators/BatchNormalizationOperator.h
deleted file mode 100644
diff --git a/Engine/Plugins/Experimental/NeuralNetworkInference/Source/NeuralNetworkInference/Public/NeuralOperators/CeilOperator.h b/Engine/Plugins/Experimental/NeuralNetworkInference/Source/NeuralNetworkInference/Public/NeuralOperators/CeilOperator.h
deleted file mode 100644
diff --git a/Engine/Plugins/Experimental/NeuralNetworkInference/Source/NeuralNetworkInference/Public/NeuralOperators/ConstantOperator.h b/Engine/Plugins/Experimental/NeuralNetworkInference/Source/NeuralNetworkInference/Public/NeuralOperators/ConstantOperator.h
deleted file mode 100644
diff --git a/Engine/Plugins/Experimental/NeuralNetworkInference/Source/NeuralNetworkInference/Public/NeuralOperators/ConvBaseOperator.h b/Engine/Plugins/Experimental/NeuralNetworkInference/Source/NeuralNetworkInference/Public/NeuralOperators/ConvBaseOperator.h
deleted file mode 100644
diff --git a/Engine/Plugins/Experimental/NeuralNetworkInference/Source/NeuralNetworkInference/Public/NeuralOperators/ConvOperator.h b/Engine/Plugins/Experimental/NeuralNetworkInference/Source/NeuralNetworkInference/Public/NeuralOperators/ConvOperator.h
deleted file mode 100644
diff --git a/Engine/Plugins/Experimental/NeuralNetworkInference/Source/NeuralNetworkInference/Public/NeuralOperators/ConvTransposeOperator.h b/Engine/Plugins/Experimental/NeuralNetworkInference/Source/NeuralNetworkInference/Public/NeuralOperators/ConvTransposeOperator.h
deleted file mode 100644
diff --git a/Engine/Plugins/Experimental/NeuralNetworkInference/Source/NeuralNetworkInference/Public/NeuralOperators/CosOperator.h b/Engine/Plugins/Experimental/NeuralNetworkInference/Source/NeuralNetworkInference/Public/NeuralOperators/CosOperator.h
deleted file mode 100644
diff --git a/Engine/Plugins/Experimental/NeuralNetworkInference/Source/NeuralNetworkInference/Public/NeuralOperators/CoshOperator.h b/Engine/Plugins/Experimental/NeuralNetworkInference/Source/NeuralNetworkInference/Public/NeuralOperators/CoshOperator.h
deleted file mode 100644
diff --git a/Engine/Plugins/Experimental/NeuralNetworkInference/Source/NeuralNetworkInference/Public/NeuralOperators/DivOperator.h b/Engine/Plugins/Experimental/NeuralNetworkInference/Source/NeuralNetworkInference/Public/NeuralOperators/DivOperator.h
deleted file mode 100644
diff --git a/Engine/Plugins/Experimental/NeuralNetworkInference/Source/NeuralNetworkInference/Public/NeuralOperators/ElementWiseOperator.h b/Engine/Plugins/Experimental/NeuralNetworkInference/Source/NeuralNetworkInference/Public/NeuralOperators/ElementWiseOperator.h
deleted file mode 100644
diff --git a/Engine/Plugins/Experimental/NeuralNetworkInference/Source/NeuralNetworkInference/Public/NeuralOperators/ExpOperator.h b/Engine/Plugins/Experimental/NeuralNetworkInference/Source/NeuralNetworkInference/Public/NeuralOperators/ExpOperator.h
deleted file mode 100644
diff --git a/Engine/Plugins/Experimental/NeuralNetworkInference/Source/NeuralNetworkInference/Public/NeuralOperators/FloorOperator.h b/Engine/Plugins/Experimental/NeuralNetworkInference/Source/NeuralNetworkInference/Public/NeuralOperators/FloorOperator.h
deleted file mode 100644
diff --git a/Engine/Plugins/Experimental/NeuralNetworkInference/Source/NeuralNetworkInference/Public/NeuralOperators/GemmOperator.h b/Engine/Plugins/Experimental/NeuralNetworkInference/Source/NeuralNetworkInference/Public/NeuralOperators/GemmOperator.h
deleted file mode 100644
diff --git a/Engine/Plugins/Experimental/NeuralNetworkInference/Source/NeuralNetworkInference/Public/NeuralOperators/LeakyReluOperator.h b/Engine/Plugins/Experimental/NeuralNetworkInference/Source/NeuralNetworkInference/Public/NeuralOperators/LeakyReluOperator.h
deleted file mode 100644
diff --git a/Engine/Plugins/Experimental/NeuralNetworkInference/Source/NeuralNetworkInference/Public/NeuralOperators/LogOperator.h b/Engine/Plugins/Experimental/NeuralNetworkInference/Source/NeuralNetworkInference/Public/NeuralOperators/LogOperator.h
deleted file mode 100644
diff --git a/Engine/Plugins/Experimental/NeuralNetworkInference/Source/NeuralNetworkInference/Public/NeuralOperators/MulOperator.h b/Engine/Plugins/Experimental/NeuralNetworkInference/Source/NeuralNetworkInference/Public/NeuralOperators/MulOperator.h
deleted file mode 100644
diff --git a/Engine/Plugins/Experimental/NeuralNetworkInference/Source/NeuralNetworkInference/Public/NeuralOperators/MultidirectionalBroadcastOperator.h b/Engine/Plugins/Experimental/NeuralNetworkInference/Source/NeuralNetworkInference/Public/NeuralOperators/MultidirectionalBroadcastOperator.h
deleted file mode 100644
diff --git a/Engine/Plugins/Experimental/NeuralNetworkInference/Source/NeuralNetworkInference/Public/NeuralOperators/NegOperator.h b/Engine/Plugins/Experimental/NeuralNetworkInference/Source/NeuralNetworkInference/Public/NeuralOperators/NegOperator.h
deleted file mode 100644
diff --git a/Engine/Plugins/Experimental/NeuralNetworkInference/Source/NeuralNetworkInference/Public/NeuralOperators/PowOperator.h b/Engine/Plugins/Experimental/NeuralNetworkInference/Source/NeuralNetworkInference/Public/NeuralOperators/PowOperator.h
deleted file mode 100644
diff --git a/Engine/Plugins/Experimental/NeuralNetworkInference/Source/NeuralNetworkInference/Public/NeuralOperators/ReciprocalOperator.h b/Engine/Plugins/Experimental/NeuralNetworkInference/Source/NeuralNetworkInference/Public/NeuralOperators/ReciprocalOperator.h
deleted file mode 100644
diff --git a/Engine/Plugins/Experimental/NeuralNetworkInference/Source/NeuralNetworkInference/Public/NeuralOperators/ReluOperator.h b/Engine/Plugins/Experimental/NeuralNetworkInference/Source/NeuralNetworkInference/Public/NeuralOperators/ReluOperator.h
deleted file mode 100644
diff --git a/Engine/Plugins/Experimental/NeuralNetworkInference/Source/NeuralNetworkInference/Public/NeuralOperators/ReshapeOperator.h b/Engine/Plugins/Experimental/NeuralNetworkInference/Source/NeuralNetworkInference/Public/NeuralOperators/ReshapeOperator.h
deleted file mode 100644
diff --git a/Engine/Plugins/Experimental/NeuralNetworkInference/Source/NeuralNetworkInference/Public/NeuralOperators/RoundOperator.h b/Engine/Plugins/Experimental/NeuralNetworkInference/Source/NeuralNetworkInference/Public/NeuralOperators/RoundOperator.h
deleted file mode 100644
diff --git a/Engine/Plugins/Experimental/NeuralNetworkInference/Source/NeuralNetworkInference/Public/NeuralOperators/SigmoidOperator.h b/Engine/Plugins/Experimental/NeuralNetworkInference/Source/NeuralNetworkInference/Public/NeuralOperators/SigmoidOperator.h
deleted file mode 100644
diff --git a/Engine/Plugins/Experimental/NeuralNetworkInference/Source/NeuralNetworkInference/Public/NeuralOperators/SignOperator.h b/Engine/Plugins/Experimental/NeuralNetworkInference/Source/NeuralNetworkInference/Public/NeuralOperators/SignOperator.h
deleted file mode 100644
diff --git a/Engine/Plugins/Experimental/NeuralNetworkInference/Source/NeuralNetworkInference/Public/NeuralOperators/SinOperator.h b/Engine/Plugins/Experimental/NeuralNetworkInference/Source/NeuralNetworkInference/Public/NeuralOperators/SinOperator.h
deleted file mode 100644
diff --git a/Engine/Plugins/Experimental/NeuralNetworkInference/Source/NeuralNetworkInference/Public/NeuralOperators/SinhOperator.h b/Engine/Plugins/Experimental/NeuralNetworkInference/Source/NeuralNetworkInference/Public/NeuralOperators/SinhOperator.h
deleted file mode 100644
diff --git a/Engine/Plugins/Experimental/NeuralNetworkInference/Source/NeuralNetworkInference/Public/NeuralOperators/SqrtOperator.h b/Engine/Plugins/Experimental/NeuralNetworkInference/Source/NeuralNetworkInference/Public/NeuralOperators/SqrtOperator.h
deleted file mode 100644
diff --git a/Engine/Plugins/Experimental/NeuralNetworkInference/Source/NeuralNetworkInference/Public/NeuralOperators/SqueezeOperator.h b/Engine/Plugins/Experimental/NeuralNetworkInference/Source/NeuralNetworkInference/Public/NeuralOperators/SqueezeOperator.h
deleted file mode 100644
diff --git a/Engine/Plugins/Experimental/NeuralNetworkInference/Source/NeuralNetworkInference/Public/NeuralOperators/SubOperator.h b/Engine/Plugins/Experimental/NeuralNetworkInference/Source/NeuralNetworkInference/Public/NeuralOperators/SubOperator.h
deleted file mode 100644
diff --git a/Engine/Plugins/Experimental/NeuralNetworkInference/Source/NeuralNetworkInference/Public/NeuralOperators/TanOperator.h b/Engine/Plugins/Experimental/NeuralNetworkInference/Source/NeuralNetworkInference/Public/NeuralOperators/TanOperator.h
deleted file mode 100644
diff --git a/Engine/Plugins/Experimental/NeuralNetworkInference/Source/NeuralNetworkInference/Public/NeuralOperators/TanhOperator.h b/Engine/Plugins/Experimental/NeuralNetworkInference/Source/NeuralNetworkInference/Public/NeuralOperators/TanhOperator.h
deleted file mode 100644
