NeuralNetworkInference (NNI) plugin:
- NNI: UNN: Warning of batch size = 1 replaced by Display.

#rnx

diff --git a/Engine/Plugins/Experimental/NeuralNetworkInference/Source/NeuralNetworkInference/Private/NeuralNetwork.cpp b/Engine/Plugins/Experimental/NeuralNetworkInference/Source/NeuralNetworkInference/Private/NeuralNetwork.cpp
--- a/Engine/Plugins/Experimental/NeuralNetworkInference/Source/NeuralNetworkInference/Private/NeuralNetwork.cpp
+++ b/Engine/Plugins/Experimental/NeuralNetworkInference/Source/NeuralNetworkInference/Private/NeuralNetwork.cpp
@@ -202,7 +202,7 @@ void UNeuralNetwork::FImpl::ConfigureTensors(FNeuralTensors& OutTensors, TArray<
 					if (CurrentTensorSize < 0)
 					{
 						CurrentTensorSizes.Push(1);
-						UE_LOG(LogNeuralNetworkInference, Warning,
+						UE_LOG(LogNeuralNetworkInference, Display,
 							TEXT("Negative (i.e., variable) dimensions not allowed yet, hard-coded to 1. Let us know if you really need variable dimensions."
 								" Keep in mind that fixed sizes might allow additional optimizations and speedup of the network during Run()."));
 					}
