NeuralNetworkInference (NNI) plugin: NNI + NNILegacy modules merged together.

Pablo, Tim, no changes needed from your side, just a FYI due to the changes in ML(Rig)Deformer.Build.Cs

#rnx

diff --git a/Engine/Plugins/Experimental/NeuralNetworkInference/NeuralNetworkInference.uplugin b/Engine/Plugins/Experimental/NeuralNetworkInference/NeuralNetworkInference.uplugin
--- a/Engine/Plugins/Experimental/NeuralNetworkInference/NeuralNetworkInference.uplugin
+++ b/Engine/Plugins/Experimental/NeuralNetworkInference/NeuralNetworkInference.uplugin
@@ -70,11 +70,6 @@
 			"WhitelistPlatforms": [
 				"Win64"
 			]
-		},
-		{
-			"Name": "NeuralNetworkInferenceLegacy",
-			"Type": "RuntimeAndProgram",
-			"LoadingPhase": "Default"
 		}
 	]
 }
\ No newline at end of file
diff --git a/Engine/Plugins/Experimental/NeuralNetworkInference/Source/Internal/NeuralNetworkInferenceLegacy/Internal/NeuralNetworkLegacy.h b/Engine/Plugins/Experimental/NeuralNetworkInference/Source/Internal/NeuralNetworkInferenceLegacy/Internal/NeuralNetworkLegacy.h
deleted file mode 100644
diff --git a/Engine/Plugins/Experimental/NeuralNetworkInference/Source/Internal/NeuralNetworkInferenceLegacy/Internal/NeuralTensorManager.h b/Engine/Plugins/Experimental/NeuralNetworkInference/Source/Internal/NeuralNetworkInferenceLegacy/Internal/NeuralTensorManager.h
deleted file mode 100644
diff --git a/Engine/Plugins/Experimental/NeuralNetworkInference/Source/Internal/NeuralNetworkInferenceLegacy/NeuralNetworkInferenceLegacy.Build.cs b/Engine/Plugins/Experimental/NeuralNetworkInference/Source/Internal/NeuralNetworkInferenceLegacy/NeuralNetworkInferenceLegacy.Build.cs
deleted file mode 100644
diff --git a/Engine/Plugins/Experimental/NeuralNetworkInference/Source/Internal/NeuralNetworkInferenceLegacy/Private/NeuralNetworkFromONNXTranslator.cpp b/Engine/Plugins/Experimental/NeuralNetworkInference/Source/Internal/NeuralNetworkInferenceLegacy/Private/NeuralNetworkFromONNXTranslator.cpp
deleted file mode 100644
diff --git a/Engine/Plugins/Experimental/NeuralNetworkInference/Source/Internal/NeuralNetworkInferenceLegacy/Private/NeuralNetworkFromONNXTranslator.h b/Engine/Plugins/Experimental/NeuralNetworkInference/Source/Internal/NeuralNetworkInferenceLegacy/Private/NeuralNetworkFromONNXTranslator.h
deleted file mode 100644
diff --git a/Engine/Plugins/Experimental/NeuralNetworkInference/Source/Internal/NeuralNetworkInferenceLegacy/Private/NeuralNetworkInferenceLegacyModule.cpp b/Engine/Plugins/Experimental/NeuralNetworkInference/Source/Internal/NeuralNetworkInferenceLegacy/Private/NeuralNetworkInferenceLegacyModule.cpp
deleted file mode 100644
diff --git a/Engine/Plugins/Experimental/NeuralNetworkInference/Source/Internal/NeuralNetworkInferenceLegacy/Private/NeuralNetworkInferenceLegacyModule.h b/Engine/Plugins/Experimental/NeuralNetworkInference/Source/Internal/NeuralNetworkInferenceLegacy/Private/NeuralNetworkInferenceLegacyModule.h
deleted file mode 100644
diff --git a/Engine/Plugins/Experimental/NeuralNetworkInference/Source/Internal/NeuralNetworkInferenceLegacy/Private/NeuralNetworkInferenceLegacyUtils.cpp b/Engine/Plugins/Experimental/NeuralNetworkInference/Source/Internal/NeuralNetworkInferenceLegacy/Private/NeuralNetworkInferenceLegacyUtils.cpp
deleted file mode 100644
diff --git a/Engine/Plugins/Experimental/NeuralNetworkInference/Source/Internal/NeuralNetworkInferenceLegacy/Private/NeuralNetworkInferenceLegacyUtils.h b/Engine/Plugins/Experimental/NeuralNetworkInference/Source/Internal/NeuralNetworkInferenceLegacy/Private/NeuralNetworkInferenceLegacyUtils.h
deleted file mode 100644
diff --git a/Engine/Plugins/Experimental/NeuralNetworkInference/Source/Internal/NeuralNetworkInferenceLegacy/Private/NeuralNetworkLegacy.cpp b/Engine/Plugins/Experimental/NeuralNetworkInference/Source/Internal/NeuralNetworkInferenceLegacy/Private/NeuralNetworkLegacy.cpp
deleted file mode 100644
diff --git a/Engine/Plugins/Experimental/NeuralNetworkInference/Source/Internal/NeuralNetworkInferenceLegacy/Private/NeuralTensorManager.cpp b/Engine/Plugins/Experimental/NeuralNetworkInference/Source/Internal/NeuralNetworkInferenceLegacy/Private/NeuralTensorManager.cpp
deleted file mode 100644
diff --git a/Engine/Plugins/Experimental/NeuralNetworkInference/Source/NeuralNetworkInference/Internal/NeuralNetworkLegacy.h b/Engine/Plugins/Experimental/NeuralNetworkInference/Source/NeuralNetworkInference/Internal/NeuralNetworkLegacy.h
new file mode 100644
--- /dev/null
+++ b/Engine/Plugins/Experimental/NeuralNetworkInference/Source/NeuralNetworkInference/Internal/NeuralNetworkLegacy.h
@@ -0,0 +1,281 @@
+// Copyright Epic Games, Inc. All Rights Reserved.
+
+#pragma once
+
+#include "CoreMinimal.h"
+#include "ModelProto.h"
+#include "NeuralEnumClasses.h"
+#include "NeuralOperator.h"
+#include "NeuralTensor.h"
+#include "NeuralTensorManager.h"
+#include "NeuralNetworkLegacy.generated.h"
+
+class UAssetImportData;
+
+/**
+ * UNeuralNetworkLegacy loads and runs a (deep) neural network.
+ *
+ * 1. Code example of UNeuralNetworkLegacy Configuration:
+ *		// -------------------------------------------------- Step 1 - Option a (Editor-only) - UNeuralNetworkLegacy Configuration From ONNX File --------------------------------------------------
+ *		#if WITH_EDITOR
+ *			// Define the UNeuralNetworkLegacy
+ *			UNeuralNetworkLegacy* Network = NewObject<UNeuralNetworkLegacy>((UObject*)GetTransientPackage(), UNeuralNetworkLegacy::StaticClass());
+ *			// Try to load the network
+ *			if (Network->Load(TEXT("SomeFilePath.onnx"))) // Alternative: Network->Load(...) + Network->IsLoaded()
+ *			{
+ *				// Set to CPU mode
+ *				Network->SetDeviceType(ENeuralDeviceType::CPU);
+ *				// Verbose example
+ *				UE_LOG(LogNeuralNetworkInference, Display, TEXT("Network loaded:\n%s."), *Network->ToString());
+ *			}
+ *		#endif
+ *
+ *		// -------------------------------------------------- Step 1 - Option b - UNeuralNetworkLegacy Configuration From UAsset File --------------------------------------------------
+ *		// Define the UNeuralNetworkLegacy
+ *		UNeuralNetworkLegacy* Network = LoadObject<UNeuralNetworkLegacy>((UObject*)GetTransientPackage(), TEXT("ExampleNetwork'/Game/Models/ExampleNetwork/ExampleNetwork.ExampleNetwork'"));
+ *		// Try to load the network
+ *		if (Network->IsLoaded())
+ *		{
+ *			// Set to CPU mode
+ *			Network->SetDeviceType(ENeuralDeviceType::CPU);
+ *			// Verbose example
+ *			UE_LOG(LogNeuralNetworkInference, Display, TEXT("Network loaded:\n%s."), *Network->ToString());
+ *		}
+ *
+ * 2. Code example of the Forward Pass:
+ *		// -------------------------------------------------- Step 2 - Forward Pass - Overview --------------------------------------------------
+ *		// Input filling
+ *		TArray<float> InArray;
+ *		Network->SetInputFromArrayCopy(InArray);
+ *		UE_LOG(LogNeuralNetworkInference, Display, TEXT("Input tensor: %s."), *Network->GetInputTensor().ToString());
+ *		// Run UNeuralNetworkLegacy
+ *		Network->Run();
+ *		// Return OutputTensor
+ *		const FNeuralTensor& OutputTensor = Network->GetOutputTensor();
+ *		UE_LOG(LogNeuralNetworkInference, Display, TEXT("Output tensor: %s."), *OutputTensor.ToString());
+ *
+ *		// -------------------------------------------------- Step 2 - Forward Pass - Input - Simple vs. Efficient and Single vs. Multiple Filling --------------------------------------------------
+ *		// Option a) Simple input filling
+ *		TArray<float> InArray;
+ *		Network->SetInputFromArrayCopy(InArray);
+ *		// Option b) Efficient input filling (it avoids the TArray-FNeuralTensor copy)
+ *		float* InputDataPointer = Network->GetInputDataPointerMutable<float>();
+ *		for (int64 Index = 0; Index < Network->GetInputTensor().Num(); ++Index)
+ *			InputDataPointer[Index] = ...;
+ *		// Option c) Simple multiple input filling
+ *		TMap<FString, FNeuralTensor> InTensorMap;
+ *		Network->SetInputFromTensorMapCopy(InTensorMap);
+ *		// Option d) Efficient multiple input filling
+ *		TMap<FString, void*> InputDataPointer = Network->CreateInputDataPointersMutable();
+ *		for (int64 Index = 0; Index < Network->GetInputTensor().Num(); ++Index)
+ *			InputDataPointer[SomeFStringKey][Index] = ...;
+ *
+ *		// -------------------------------------------------- Step 2 - Forward Pass - Output - Single vs. Multiple Output --------------------------------------------------
+ *		// Option a) Return OutputTensor
+ *		const FNeuralTensor* OutputTensor = Network->GetOutputTensor();
+ *		UE_LOG(LogNeuralNetworkInference, Display, TEXT("Output tensor: %s."), *OutputTensor.ToString());
+ *		// Option b) Return OutputTensor
+ *		const TArray<FNeuralTensor>& Tensors = GetTensors();
+ *		const TMap<FString, int32>& OutputNameIndexMap = Network->GetOutputNameIndexMap();
+ *		// Iterate and print on console all output tensors
+ *		for (const auto& OutputNameIndexPair)
+ *		{
+ *			UE_LOG(LogNeuralNetworkInference, Display, TEXT("Output tensor %s: %s."), *OutputNameIndexPair.Key, *Tensors[OutputNameIndexPair.Key].ToString());
+ *		}
+ */
+//UCLASS(BlueprintType, Deprecated)
+//class NEURALNETWORKINFERENCELEGACY_API UDEPRECATED_UNeuralNetworkLegacy : public UObject
+UCLASS(BlueprintType)
+class NEURALNETWORKINFERENCE_API UNeuralNetworkLegacy : public UObject
+{
+	GENERATED_BODY()
+
+public:
+	/**
+	 * Callbacks when the asyn Run() is finished. It will happen from any thread.
+	 */
+	DECLARE_DELEGATE(FOnAsyncRunCompletedInAnyThread);
+
+	/**
+	 * This enum only applies to GPU mode, where Run() might block the thread until completed, or might run on a background thread.
+	 */
+	enum class EGPUSynchronousMode : uint8
+	{
+		Synchronous, /* Block the thread until the function (e.g., Run) has finished. */
+		Asynchronous /* Do not block the thread (the function will run in the background, e.g., Run). Delegates would let the user know when the function has finished. */
+	};
+
+	UNeuralNetworkLegacy();
+
+	virtual ~UNeuralNetworkLegacy();
+
+	//~UObject interface
+	virtual void PostInitProperties() override;
+	virtual void PostLoad() override;
+	virtual void Serialize(FArchive& Archive) override;
+	//~End of UObject interface
+#if WITH_EDITOR
+	/** Editor-only function: Re-import asset with editor data (imported file). */
+	void ReimportAssetFromEditorData();
+	/** Editor-only function: Importing data and options used for loading the neural network. */
+	UAssetImportData* GetAssetImportData() const;
+	UAssetImportData* GetAndMaybeCreateAssetImportData();
+#endif // WITH_EDITOR
+
+	/**
+	 * Editor-only function
+	 * It loads the desired UNeuralNetworkLegacy definition and its weights from an ONNX file.
+	 * @return Whether it loaded the UNeuralNetworkLegacy successfully.
+	 */
+#if WITH_EDITOR
+	bool Load(const FString& InFilePath);
+#endif // WITH_EDITOR
+
+	/**
+	 * It loads the desired UNeuralNetworkLegacy definition and its weights from an deserialized UNeuralNetworkLegacy uasset.
+	 * @param InTensorManager It will be moved for performance reasons. Do not use InTensorManager after calling this function.
+	 * @return Whether it loaded the UNeuralNetworkLegacy successfully.
+	 */
+	bool Load(FNeuralTensorManager& InTensorManager, const TArray<TSharedPtr<FNeuralOperator>>& InOperators);
+
+	/**
+	 * It returns whether it loaded the UNeuralNetworkLegacy successfully.
+	 */
+	bool IsLoaded() const;
+
+	/**
+	 * Whether Run() will occur on the CPU or GPU.
+	 * If SetDeviceType() is never called, the default device (EDeviceType::CPU) will be used.
+	 */
+	ENeuralDeviceType GetDeviceType() const;
+	void SetDeviceType(const ENeuralDeviceType InDeviceType);
+
+	/**
+	 * Returns the delegate called when async Run() is over (on any thread)
+	 */
+	FOnAsyncRunCompletedInAnyThread& GetOnAsyncRunCompletedInAnyThreadDelegate();
+
+	/**
+	 * It returns the TArray of FNeuralTensors as a read-only object.
+	 */
+	const TArray<FNeuralTensor>& GetTensors() const;
+
+	/**
+	 * There are 6 alternative functions to fill the input FNeuralTensor(s) data:
+	 * If exactly 1 input FNeuralTensor: SetInputFromArrayCopy(), SetInputFromTensorCopy(), GetInputDataPointerMutable().
+	 *     - They ensure there is exactly 1 input FNeuralTensor or will log a warning if more than 1 input tensor exists.
+	 * If more than 1 input FNeuralTensors: SetInputFromTensorMapCopy(), CreateInputDataPointersMutable().
+	 *
+	 * - SetInputFromArrayCopy()/SetInputFromTensorCopy()/SetInputFromTensorMapCopy() deeply copy the input FNeuralTensor(s) (slower but safer and less error
+	 *   prone). See FNeuralTensor::SetFromArrayCopy() for more details.
+	 * - GetInputDataPointerMutable() returns a pointer to the input FNeuralTensor raw data, so it can be filled before calling Run().
+	 * - CreateInputDataPointersMutable() returns a TMap of pointers to the input FNeuralTensors raw data, so it can be filled before calling Run().
+	 *
+	 * @return GetInputDataPointerMutable() returns a T* corresponding to the mutable data in the single input FNeuralTensor, CreateInputDataPointersMutable() returns a
+	 * TMap<FString, void*> for each one of the input FNeuralTensors.
+	 *
+	 * For read-only access to the input FNeuralTensor(s), see GetInputTensor() or GetInputNameIndexMap() (e.g., to extract properties
+	 * from the input FNeuralTensor(s) such as volume, dimensions).
+	 */
+	template<typename T>
+	void SetInputFromArrayCopy(const TArray<T>& InArray);
+	void SetInputFromTensorCopy(const FNeuralTensor& InTensor);
+	void SetInputFromTensorMapCopy(const TMap<FString, FNeuralTensor>& InTensorMap);
+	template<typename T>
+	T* GetInputDataPointerMutable();
+	TMap<FString, void*> CreateInputDataPointersMutable();
+	FRDGBufferUAVRef GetInputBufferUAVRef();
+	TMap<FString, FRDGBufferUAVRef> CreateInputBufferUAVRefs();
+
+	/**
+	 * It returns the input FNeuralTensor (map) as a read-only object.
+	 * GetInputTensor() ensures there is only 1 input and returns a single TSharedPtr<FNeuralTensor>, while GetInputNameIndexMap() returns the whole TMap.
+	 * In order to modify their input data of the FNeuralTensor (map), use GetInputDataPointerMutable() or CreateInputDataPointersMutable().
+	 *
+	 * These 2 functions are only meant to extract the input properties (e.g., volume, dimensions), but the actual FNeuralTensor(map) cannot be modified
+	 * to avoid undefined behavior in the next Run() of the UNeuralNetworkLegacy due to non-controlled functions (e.g., resizes, memory re-allocation).
+	 *
+	 * @return GetInputTensor() returns a read-only const FNeuralTensor& corresponding to the single input FNeuralTensor, GetInputNameIndexMap() returns a
+	 * TMap<FString, int32> for each one of the input FNeuralTensors. The TMap will go out of scope if UNeuralNetworkLegacy does.
+	 */
+	const FNeuralTensor& GetInputTensor() const;
+	const TMap<FString, int32>& GetInputNameIndexMap() const;
+
+	/**
+	 * It returns the output FNeuralTensor (map) as a read-only object.
+	 * GetOutputTensor() ensures there is only 1 output and returns a single TSharedPtr<FNeuralTensor>, while GetOutputNameIndexMap() returns the whole TMap.
+	 *
+	 * These 2 functions are only meant to extract the output properties (e.g., output result, volume, dimensions), but the actual FNeuralTensor(map) cannot
+	 * be modified to avoid undefined behavior in the next Run() of the UNeuralNetworkLegacy due to non-controlled functions (e.g., resizes, memory re-allocation).
+	 *
+	 * @return GetOutputTensor() returns a read-only const FNeuralTensor& corresponding to the single output FNeuralTensor, GetOutputNameIndexMap() returns a
+	 * TMap<FString, int32> for each one of the input FNeuralTensors. The TMap will go out of scope if UNeuralNetworkLegacy does.
+	 */
+	const FNeuralTensor& GetOutputTensor() const;
+	const TMap<FString, int32>& GetOutputNameIndexMap() const;
+	const FRDGBufferSRVRef GetOutputBufferSRVRef() const;
+	TMap<FString, const FRDGBufferSRVRef> CreateOutputBufferSRVRefs() const;
+
+	/**
+	 * These functions are slower than the Get ones, because they deep copy each one of the TArrays in the final TMap.
+	 */
+	TMap<FString, FNeuralTensor> CreateInputTensorMap() const;
+	TMap<FString, FNeuralTensor> CreateOutputTensorMap() const;
+
+	/**
+	 * Run() executes the forward pass on the current UNeuralNetworkLegacy given the current input FNeuralTensor(s), which were previously filled with
+	 * SetInputFromTensorCopy() or GetInputDataPointerMutable(), or their multi-input analogs.
+	 * Its output results can be retrieved with GetOutputTensor() or GetOutputNameIndexMap().
+	 * @param GPUSynchronousMode Whether it should block the thread until the UNeuralNetworkLegacy has fully run.
+	 */
+	void Run(const EGPUSynchronousMode InGPUSynchronousMode = EGPUSynchronousMode::Synchronous, const ENeuralDeviceType InInputDeviceType = ENeuralDeviceType::CPU, const ENeuralDeviceType InOutputDeviceType = ENeuralDeviceType::CPU,
+		const bool bRunGPUEmptyOnlyForProfiling = false);
+
+	/**
+	 * Auxiliary function that returns the network architecture, weights, tensors, operators, etc.
+	 */
+	FString ToString() const;
+
+protected:
+	UPROPERTY(VisibleAnywhere, Category = "Neural Network Inference")
+	TArray<int32> Version;
+	UPROPERTY(VisibleAnywhere, Category = "Neural Network Inference")
+	bool bIsLoaded;
+	UPROPERTY(EditAnywhere, Category = "Neural Network Inference")
+	ENeuralDeviceType DeviceType;
+	UPROPERTY(VisibleAnywhere, Category = "Neural Network Inference")
+	FNeuralTensorManager TensorManager; /* It contains a few TArray and TMaps for all FNeuralTensors (Input, Output, Intermediate(Not)Initialized, Weight) */
+	UPROPERTY(VisibleAnywhere, Category = "Neural Network Inference")
+	FModelProto ModelProto;
+#if WITH_EDITORONLY_DATA
+	/** Importing data and options used for loading the neural network. */
+	UPROPERTY(VisibleAnywhere, Instanced, Category = "ImportSettings")
+	UAssetImportData* AssetImportData;
+#endif // WITH_EDITORONLY_DATA
+
+private:
+	// Non-uproperty members
+	bool bAreTensorsInGpu; /* It should always be false when loaded from uasset (FNeuralTensors are not auto-loaded to GPU) */
+	/**
+	 * Operators represents the set of operators of the network that need to run on the Forward pass and that might need to run on the PostForward.
+	 */
+	TArray<TSharedPtr<FNeuralOperator>> Operators;
+	FOnAsyncRunCompletedInAnyThread OnAsyncRunCompletedInAnyThreadDelegate;
+};
+
+
+
+/* UNeuralNetworkLegacy template functions
+ *****************************************************************************/
+
+template<typename T>
+void UNeuralNetworkLegacy::SetInputFromArrayCopy(const TArray<T>& InArray)
+{
+	return TensorManager.SetInputFromArrayCopy(InArray);
+}
+
+template<typename T>
+T* UNeuralNetworkLegacy::GetInputDataPointerMutable()
+{
+	return TensorManager.GetInputDataPointerMutable<T>();
+}
diff --git a/Engine/Plugins/Experimental/NeuralNetworkInference/Source/NeuralNetworkInference/Internal/NeuralTensorManager.h b/Engine/Plugins/Experimental/NeuralNetworkInference/Source/NeuralNetworkInference/Internal/NeuralTensorManager.h
new file mode 100644
--- /dev/null
+++ b/Engine/Plugins/Experimental/NeuralNetworkInference/Source/NeuralNetworkInference/Internal/NeuralTensorManager.h
@@ -0,0 +1,197 @@
+// Copyright Epic Games, Inc. All Rights Reserved.
+
+#pragma once
+
+#include "CoreMinimal.h"
+#include "NeuralTensor.h"
+#include "NeuralTensorManager.generated.h"
+
+/**
+ * This is an auxiliary class. See UNeuralNetworkLegacy for a high-level wrapper of the whole NeuralNetworkInference plugin. The UNeuralNetworkLegacy header
+ * documentation also includes some code examples.
+ *
+ * FNeuralTensorManager is an utility class that helps store all the FNeuralTensors related to a particular UNeuralNetworkLegacy model.
+ * Most functions are inlined (or templated) to avoid the performance overhead of a wrapper class.
+ */
+USTRUCT()
+struct NEURALNETWORKINFERENCE_API FNeuralTensorManager
+{
+	GENERATED_BODY()
+
+public:
+	FNeuralTensorManager();
+
+	/**
+	 * Calls FNeuralTensorManager() + Load(InNameIndexMap, InTensors, InInputNameIndexMap, InOutputNameIndexMap).
+	 */
+	FNeuralTensorManager(TArray<FNeuralTensor>& InTensors, TMap<FString, int32>& InNameIndexMap, TMap<FString, int32>& InInputNameIndexMap,
+		TMap<FString, int32>& InOutputNameIndexMap);
+
+	/**
+	 * Calls FNeuralTensorManager() + Load(InNameIndexMap, InTensors, InInputTensors, InOutputTensors).
+	 */
+	FNeuralTensorManager(TArray<FNeuralTensor>& InTensors, const TArray<FNeuralTensor*>& InInputTensors, const TArray<FNeuralTensor*>& InOutputTensors);
+
+	~FNeuralTensorManager();
+
+	/**
+	 * It moves the TMap objects and populates all other class members.
+	 */
+	bool Load(TArray<FNeuralTensor>& InTensors, TMap<FString, int32>& InNameIndexMap, TMap<FString, int32>& InInputNameIndexMap,
+		TMap<FString, int32>& InOutputNameIndexMap);
+
+	/**
+	 * Alternative to the original Load(), recommended for single-input single-output models as well as any kind of network whose
+	 * input and output tensor names are irrelevant). It will internally create InNameIndexMap/InInputNameIndexMap/InOutputNameIndexMap
+	 * and call the original Load().
+	 * @param InInputTensors Pointers for the tensors in InTensors that are inputs.
+	 * @param InOutputTensors Pointers for the tensors in InTensors that are outputs.
+	 */
+	bool Load(TArray<FNeuralTensor>& InTensors, const TArray<FNeuralTensor*>& InInputTensors, const TArray<FNeuralTensor*>& InOutputTensors);
+
+	/**
+	 * It returns whether it loaded the UNeuralNetworkLegacy successfully.
+	 */
+	bool IsLoaded() const;
+	TArray<int32> GetVersion() const;
+
+	/**
+	 * Getter for the Tensors TArray. GetTensors() will return a const reference while GetTensorsMutable() allows modifying it.
+	 */
+	const TArray<FNeuralTensor>& GetTensors() const;
+	TArray<FNeuralTensor>& GetTensorsMutable();
+
+	/**
+	 * Getter for the NameIndexMap TMap.
+	 */
+	const TMap<FString, int32>& GetNameIndexMap() const;
+
+	/**
+	 * It returns the non-input FNeuralTensor index TMap as a read-only object.
+	 */
+	const TArray<int32>& GetNonInputIndexes() const;
+
+	/**
+	 * It returns the input FNeuralTensor(s) as read-only (GetInputTensor/GetInputIndexes/GetInputNameIndexMap) or mutable (GetInputTensorMutable) objects.
+	 * GetInputTensor/GetInputTensorMutable ensure there is only 1 input and returns a single FNeuralTensor&, while GetInputIndexes/GetInputNameIndexMap return them all.
+	 * In order to modify their input data of the FNeuralTensor (map), use GetInputDataPointerMutable() or CreateInputDataPointersMutable().
+	 *
+	 * Use the read-only versions as much as possible to avoid potential undefined behavior in the next Run() of the FNeuralTensorManager due to non-controlled functions
+	 * (e.g., resizes, memory re-allocation).
+	 *
+	 * @return GetInputTensor() returns a read-only const FNeuralTensor& corresponding to the single input FNeuralTensor, GetInputIndexes()/GetInputNameIndexMap() return a
+	 * TArray<int32>/TMap<FString, int32> for each one of the input FNeuralTensors. They will all go out of scope if FNeuralTensorManager does.
+	 */
+	const FNeuralTensor& GetInputTensor() const;
+	FNeuralTensor& GetInputTensorMutable();
+	const TArray<int32>& GetInputIndexes() const;
+	const TMap<FString, int32>& GetInputNameIndexMap() const;
+
+	/**
+	 * There are 6 alternative functions to fill the input FNeuralTensor(s) data:
+	 * If exactly 1 input FNeuralTensor: SetInputFromArrayCopy(), SetInputFromTensorCopy(), GetInputDataPointerMutable().
+	 *     - They ensure there is exactly 1 input FNeuralTensor or will log a warning if more than 1 input tensor exists.
+	 * If more than 1 input FNeuralTensors: SetInputFromTensorMapCopy(), CreateInputDataPointersMutable().
+	 *
+	 * - SetInputFromArrayCopy()/SetInputFromTensorCopy()/SetInputFromTensorMapCopy() deeply copy the input FNeuralTensor(s) (slower but safer and less error
+	 *   prone). See FNeuralTensor::SetFromArrayCopy() for more details.
+	 * - GetInputDataPointerMutable() returns a pointer to the input FNeuralTensor raw data, so it can be filled before calling Run().
+	 * - CreateInputDataPointersMutable() returns a TMap of pointers to the input FNeuralTensors raw data, so it can be filled before calling Run().
+	 *
+	 * @return GetInputDataPointerMutable() returns a T* corresponding to the mutable data in the single input FNeuralTensor, CreateInputDataPointersMutable() returns a
+	 * TMap<FString, void*> for each one of the input FNeuralTensors.
+	 *
+	 * For read-only access to the input FNeuralTensor(s), see GetInputTensor() or GetInputNameIndexMap() (e.g., to extract properties
+	 * from the input FNeuralTensor(s) such as volume, dimensions).
+	 */
+	template<typename T>
+	void SetInputFromArrayCopy(const TArray<T>& InArray);
+	void SetInputFromTensorCopy(const FNeuralTensor& InTensor);
+	void SetInputFromTensorMapCopy(const TMap<FString, FNeuralTensor>& InTensorMap);
+	template<typename T>
+	T* GetInputDataPointerMutable();
+	TMap<FString, void*> CreateInputDataPointersMutable();
+	FRDGBufferUAVRef GetInputBufferUAVRef();
+	TMap<FString, FRDGBufferUAVRef> CreateInputBufferUAVRefs();
+
+	/**
+	 * It returns the input FNeuralTensor(s) as read-only (GetOutputTensor/GetOutputIndexes/GetOutputNameIndexMap) or mutable (GetOutputTensorMutable) objects.
+	 * GetOutputTensor/GetOutputTensorMutable ensure there is only 1 input and returns a single FNeuralTensor&, while GetOutputIndexes/GetOutputNameIndexMap return them all.
+	 * In order to modify their input data of the FNeuralTensor (map), use GetOutputDataPointerMutable() or CreateOutputDataPointersMutable().
+	 *
+	 * Use the read-only versions as much as possible to avoid potential undefined behavior in the next Run() of the FNeuralTensorManager due to non-controlled functions
+	 * (e.g., resizes, memory re-allocation).
+	 *
+	 * @return GetOutputTensor() returns a read-only const FNeuralTensor& corresponding to the single input FNeuralTensor, GetOutputIndexes()/GetOutputNameIndexMap() return a
+	 * TArray<int32>/TMap<FString, int32> for each one of the input FNeuralTensors. They will all go out of scope if FNeuralTensorManager does.
+	 */
+	const FNeuralTensor& GetOutputTensor() const;
+	FNeuralTensor& GetOutputTensorMutable();
+	const TArray<int32>& GetOutputIndexes() const;
+	const TMap<FString, int32>& GetOutputNameIndexMap() const;
+	const FRDGBufferSRVRef GetOutputBufferSRVRef() const;
+	TMap<FString, const FRDGBufferSRVRef> CreateOutputBufferSRVRefs() const;
+
+	/**
+	 * These functions are slower than the rest, because they deep copy each one of the TArray's in the final TMap.
+	 */
+	TMap<FString, FNeuralTensor> CreateInputTensorMap() const;
+	TMap<FString, FNeuralTensor> CreateOutputTensorMap() const;
+
+protected:
+	UPROPERTY(VisibleAnywhere, Category = "Neural Network Inference")
+	TArray<int32> Version;
+	UPROPERTY(VisibleAnywhere, Category = "Neural Network Inference")
+	bool bIsLoaded;
+	/**
+	 * It contains the raw weight, intermediate, input, and output FNeuralTensor's. Inlined FNeuralTensor's will only appear once on Tensors.
+	 * E.g., in the example described for NameIndexMap, the FNeuralTensor will be present with the Key name "Input" (or otherwise "Output"). E.g.,
+	 * FNeuralTensor& Tensor = Tensors[NameIndexMap.FindChecked(TEXT("Input"))];
+	 */
+	UPROPERTY(VisibleAnywhere, Category = "Neural Network Inference")
+	TArray<FNeuralTensor> Tensors;
+	/**
+	 * For each original FNeuralTensor name, it saves its final name after merging FNeuralTensor's for memory savings.
+	 * E.g., For an inlined ReLU operator with original input tensor name "Input" and output tensor name "Output", NameIndexMap will
+	 * contain the following 2 fields: {"Input", 0} and {"Output", 0}.
+	 * E.g., For a non-inlined ReLU: {"Input", 0} and {"Output", 1}.
+	 */
+	UPROPERTY(VisibleAnywhere, Category = "Neural Network Inference")
+	TMap<FString, int32> NameIndexMap;
+	/**
+	 * Other variables given in Load().
+	 */
+	UPROPERTY(VisibleAnywhere, Category = "Neural Network Inference")
+	TMap<FString, int32> InputNameIndexMap;
+	UPROPERTY(VisibleAnywhere, Category = "Neural Network Inference")
+	TMap<FString, int32> OutputNameIndexMap;
+	/**
+	 * Other variables auto-generated in Load().
+	 */
+	UPROPERTY(VisibleAnywhere, Category = "Neural Network Inference")
+	TArray<int32> InputIndexes;
+	UPROPERTY(VisibleAnywhere, Category = "Neural Network Inference")
+	TArray<int32> NonInputIndexes;
+	UPROPERTY(VisibleAnywhere, Category = "Neural Network Inference")
+	TArray<int32> OutputIndexes;
+
+	bool SanityCheckIsLoaded() const;
+};
+
+
+
+/* FNeuralTensorManager inline functions
+ *****************************************************************************/
+
+template<typename T>
+void FNeuralTensorManager::SetInputFromArrayCopy(const TArray<T>& InArray)
+{
+	return GetInputTensorMutable().SetFromArrayCopy(InArray);
+}
+
+template<typename T>
+T* FNeuralTensorManager::GetInputDataPointerMutable()
+{
+	// GetInputTensorMutable() already runs the sanity checks
+	return GetInputTensorMutable().GetDataCasted<T>();
+}
diff --git a/Engine/Plugins/Experimental/NeuralNetworkInference/Source/NeuralNetworkInference/NeuralNetworkInference.Build.cs b/Engine/Plugins/Experimental/NeuralNetworkInference/Source/NeuralNetworkInference/NeuralNetworkInference.Build.cs
--- a/Engine/Plugins/Experimental/NeuralNetworkInference/Source/NeuralNetworkInference/NeuralNetworkInference.Build.cs
+++ b/Engine/Plugins/Experimental/NeuralNetworkInference/Source/NeuralNetworkInference/NeuralNetworkInference.Build.cs
@@ -28,7 +28,9 @@ public class NeuralNetworkInference : ModuleRules
 				"Core",
 				"CoreUObject",
 				"Engine",
-				"NeuralNetworkInferenceCore"
+				"NeuralNetworkInferenceCore",
+				// Internal (Backend-related)
+				"ModelProto"
 			}
 		);
 
@@ -36,7 +38,6 @@ public class NeuralNetworkInference : ModuleRules
 			(
 			new string[] {
 				// Backend-related
-				"ModelProto",
 				"NeuralNetworkInferenceShaders",
 				"RenderCore",
 				"RHI",
@@ -46,6 +47,21 @@ public class NeuralNetworkInference : ModuleRules
 			}
 		);
 
+		// Editor-only
+		if (Target.bBuildEditor)
+		{
+			// Win64-only
+			if (Target.Platform == UnrealTargetPlatform.Win64)
+			{
+				PrivateDependencyModuleNames.AddRange
+					(
+					new string[] {
+						"ModelProtoFileReader"
+					}
+				);
+			}
+		}
+
 		if (bIsORTSupported)
 		{
 
diff --git a/Engine/Plugins/Experimental/NeuralNetworkInference/Source/NeuralNetworkInference/Private/NeuralNetworkFromONNXTranslator.cpp b/Engine/Plugins/Experimental/NeuralNetworkInference/Source/NeuralNetworkInference/Private/NeuralNetworkFromONNXTranslator.cpp
new file mode 100644
--- /dev/null
+++ b/Engine/Plugins/Experimental/NeuralNetworkInference/Source/NeuralNetworkInference/Private/NeuralNetworkFromONNXTranslator.cpp
@@ -0,0 +1,614 @@
+// Copyright Epic Games, Inc. All Rights Reserved.
+
+#include "NeuralNetworkFromONNXTranslator.h"
+#include "Misc/Paths.h"
+#include "NeuralNetworkInferenceUtils.h"
+#include "NeuralOperator.h"
+#include "NeuralOperators.h"
+#include "NeuralTensorManager.h"
+
+#if WITH_EDITOR
+#if PLATFORM_WINDOWS
+#include "ModelProtoFileReader.h"
+#endif //PLATFORM_WINDOWS
+#endif //WITH_EDITOR
+
+
+
+/* FNeuralNetworkFromONNXTranslator public functions
+ *****************************************************************************/
+
+bool FNeuralNetworkFromONNXTranslator::Translate(TArray<TSharedPtr<FNeuralOperator>>& OutOperators, FNeuralTensorManager& InTensorManager, const FGraphProto& InGraphProto)
+{
+	TMap<FString, int32> NameIndexMap = InTensorManager.GetNameIndexMap();
+	TMap<FString, int32> DummyOutputNameIndexMap;
+	const bool bResult = CreateOperatorsAndEditTensorArray(OutOperators, DummyOutputNameIndexMap, InTensorManager.GetTensorsMutable(), NameIndexMap, InTensorManager.GetInputNameIndexMap(),
+		InTensorManager.GetOutputNameIndexMap(), InGraphProto);
+	if (NameIndexMap.Num() != InTensorManager.GetNameIndexMap().Num())
+	{
+		UE_LOG(LogNeuralNetworkInference, Warning, TEXT("FNeuralNetworkFromONNXTranslator::Translate(): NameIndexMap changed its size, it must not (%d != %d)."),
+			NameIndexMap.Num(), InTensorManager.GetNameIndexMap().Num());
+	}
+	return bResult;
+}
+
+#if WITH_EDITOR
+bool FNeuralNetworkFromONNXTranslator::Translate(TArray<TSharedPtr<FNeuralOperator>>& OutOperators, FNeuralTensorManager& OutTensorManager, const FGraphProto& InGraphProto,
+	const FString& InModelPath)
+{
+	// MaxNumberTensorsInNetwork is required to avoid the bug of resizing and invalidating all pointers given to the operators
+	int32 MaxNumberTensorsInNetwork = 0;
+	{
+		// Quickly estimate the number of FNeuralTensor's of the UNeuralNetworkLegacy for Tensors.Reserve(...)
+		// Number FNeuralTensor's = Number (unique) input FNeuralTensors + number (unique) output FNeuralTensors
+		// Uniqueness guaranteed by using the TSet
+		TSet<FString> TensorNameSet;
+		for (const FNodeProto& NodeProto : InGraphProto.Node)
+		{
+			// Check input tensors
+			for (const FString& TensorName : NodeProto.Input)
+			{
+				if (TensorNameSet.Find(TensorName) == nullptr)
+				{
+					TensorNameSet.Add(TensorName);
+				}
+			}
+			// Check output tensors
+			for (const FString& TensorName : NodeProto.Output)
+			{
+				if (TensorNameSet.Find(TensorName) == nullptr)
+				{
+					TensorNameSet.Add(TensorName);
+				}
+			}
+		}
+		MaxNumberTensorsInNetwork = TensorNameSet.Num()
+		// Add maximum number possible of auxiliary tensors. A conservative estimate is the number of operators times the maximum possible number of auxiliary tensors per operator
+			+ InGraphProto.Node.Num() * FNeuralOperator::GetMaximumPossibleNumberAuxiliaryTensorsPerOperator();
+	}
+	TArray<FNeuralTensor> Tensors;
+	Tensors.Reserve(MaxNumberTensorsInNetwork);
+	TMap<FString, int32> NameIndexMap;
+
+	// Read input tensors
+	TMap<FString, int32> InputNameIndexMap;
+	for (const FValueInfoProto& GraphProtoInput : InGraphProto.Input)
+	{
+		// Create tensor
+		const TArray<FTensorShapeProtoDimension>& Dimensions = GraphProtoInput.Type.TensorType.Shape.Dim;
+		TArray<int64> DimensionsAsTArray;
+		for (const FTensorShapeProtoDimension& Dimension : Dimensions)
+		{
+			DimensionsAsTArray.Emplace(Dimension.DimValue);
+		}
+		// Add operator name to TMap
+		NameIndexMap.Add(GraphProtoInput.Name, Tensors.Num());
+		InputNameIndexMap.Add(GraphProtoInput.Name, Tensors.Num());
+		// Add tensor to TMap
+		Tensors.Emplace(FNeuralTensor(ENeuralDataType::Float, DimensionsAsTArray, GraphProtoInput.Name, ENeuralTensorTypeGPU::Input));
+	}
+
+	// Read output tensor names (but do not create them yet, it has to happen simultaneously with the operator creation and tensor inlining)
+	// Issue: For inlining to work properly, output tensors cannot be created until after inlining operators (i.e., CreateOperatorsAndEditTensorArray).
+	// Otherwise, the following failing case could happen with a network like: InputTensor --> Gemm --> IntermediateTensor --> Relu --> OutputTensor
+	// 1. Inputs and outputs are created before hand, i.e., InputTensor and OutputTensor
+	// 2. CreateOperatorsAndEditTensorArray() is run
+	//   2.1. Gemm and IntermediateTensor would be created (because Gemm cannot be inlined)
+	//   2.2. Relu would be created, Relu would verify its input can be inlined, but both IntermediateTensor and OutputTensor are already created at this point, so it cannot inline them (crash/bug)
+	TMap<FString, int32> OutputNameDummyIndexMap;
+	for (const FValueInfoProto& GraphProtoOutput : InGraphProto.Output)
+	{
+		// Add operator name to TMap
+		if (NameIndexMap.Find(GraphProtoOutput.Name))
+		{
+			UE_LOG(LogNeuralNetworkInference, Warning, TEXT("FNeuralNetworkFromONNXTranslator::Translate(): A FNeuralTensor of the UNeuralNetworkLegacy is a global Input and Output simultaneously."));
+		}
+		OutputNameDummyIndexMap.Add(GraphProtoOutput.Name, -1);
+	}
+
+	// Fill operators
+	TMap<FString, int32> OutputNameIndexMap;
+	bool bWasSuccessful = CreateOperatorsAndEditTensorArray(OutOperators, OutputNameIndexMap, Tensors, NameIndexMap, InputNameIndexMap, OutputNameDummyIndexMap, InGraphProto, InModelPath);
+	if (Tensors.Num() <= NameIndexMap.Num())
+	{
+		UE_LOG(LogNeuralNetworkInference, Display, TEXT("Auto-inlining: %d initial tensors compressed into %d (%d tensors merged)."),
+			NameIndexMap.Num(), Tensors.Num(), NameIndexMap.Num() - Tensors.Num());
+	}
+	// Sanity check
+	if (Tensors.Num() > MaxNumberTensorsInNetwork)
+	{
+		UE_LOG(LogNeuralNetworkInference, Warning, TEXT("FNeuralNetworkFromONNXTranslator::Translate(): Tensors.Num() <= MaxNumberTensorsInNetwork failed (%d vs. %d). A resize of Tensors"
+			" might have occurred, invalidating all TArray<FNeuralTensor*> FOperator variables. Report this error (so we can fix MaxNumberTensorsInNetwork to be big enough)."),
+			Tensors.Num(), MaxNumberTensorsInNetwork);
+		bWasSuccessful = false;
+	}
+
+	// Set OutTensorManager
+	if (bWasSuccessful)
+	{
+		return OutTensorManager.Load(Tensors, NameIndexMap, InputNameIndexMap, OutputNameIndexMap);
+	}
+	return bWasSuccessful;
+}
+#endif //WITH_EDITOR
+
+
+
+/* FNeuralNetworkFromONNXTranslator private functions
+ *****************************************************************************/
+
+bool FNeuralNetworkFromONNXTranslator::CreateOperatorsAndEditTensorArray(TArray<TSharedPtr<FNeuralOperator>>& OutOperators, TMap<FString, int32>& OutputNameIndexMap,
+	TArray<FNeuralTensor>& InOutTensors, TMap<FString, int32>& InOutNameIndexMap, const TMap<FString, int32>& InInputNameDummyIndexMap,
+	const TMap<FString, int32>& InOutputNameDummyIndexMap, const FGraphProto& InGraphProto, const FString& InModelPath)
+{
+	const bool bConvertingFromSerializedData = (InModelPath.Len() == 0);
+	const int64 InOutTensorsNumInit = InOutTensors.Num();
+	const int64 InOutNameIndexMapNumInit = InOutNameIndexMap.Num();
+	// Create operators sequentially
+	OutOperators.Empty();
+	OutputNameIndexMap.Empty();
+	// Iterating over each operator
+	for (const FNodeProto& NodeProto : InGraphProto.Node)
+	{
+		// Read/set input tensors
+		TArray<FNeuralTensor*> OperatorInputTensors;
+		TArray<FString> OperatorInputTensorNames;
+		// Iterating over each input tensor name for that operator
+		for (const FString& InputTensorName : NodeProto.Input)
+		{
+			OperatorInputTensorNames.Push(InputTensorName);
+			// FNeuralTensor already read (e.g., absolute input of the UNeuralNetworkLegacy, or the output of a previous tensor required as input of this one)
+			if (int32* ExistingTensorIndex = InOutNameIndexMap.Find(InputTensorName))
+			{
+				OperatorInputTensors.Push(&InOutTensors[*ExistingTensorIndex]);
+			}
+			// FNeuralTensor was never modified, so it must come from a weight tensor (regardless of whether it is also an absolute output of the network)
+			else
+			{
+#if WITH_EDITOR
+				const FTensorProto& TensorProto = *FModelProto::FindElementInArray(InputTensorName, InGraphProto.Initializer, /*bMustValueBeFound*/true);
+				// Is it also an absolute output of the network?
+				const int64 TensorIndex = InOutTensors.Num();
+				if (!bConvertingFromSerializedData && InOutputNameDummyIndexMap.Find(InputTensorName))
+				{
+					OutputNameIndexMap.Add(InputTensorName, TensorIndex);
+				}
+				// Create and add tensor to InOutTensors
+				InOutTensors.Push(FNeuralTensor());
+				FNeuralTensor& NewTensor = InOutTensors.Last();
+				// Deprecated code for otxt files, if added back, add before the 2 for loops for speedup
+				// FString PathPart, FilenamePart, ExtensionPart;
+				// if (!bConvertingFromSerializedData)
+				// {
+				// 	FPaths::Split(InModelPath, PathPart, FilenamePart, ExtensionPart);
+				// }
+				// Get NewTensor from TensorProto
+				if (!NewTensor.SetFromTensorProto(&TensorProto, InputTensorName, ENeuralTensorTypeGPU::Weight/*, PathPart*/))
+				{
+					UE_LOG(LogNeuralNetworkInference, Warning, TEXT("FNeuralNetworkFromONNXTranslator::CreateOperatorsAndEditTensorArray(): SetFromTensorProto() failed."));
+					return false;
+				}
+				// Update InOutNameIndexMap and OperatorInputTensors
+				InOutNameIndexMap.Add(InputTensorName, TensorIndex);
+				OperatorInputTensors.Push(&NewTensor);
+#else
+				UE_LOG(LogNeuralNetworkInference, Warning, TEXT("FNeuralNetworkFromONNXTranslator::CreateOperatorsAndEditTensorArray(): This code should never be called in non-Editor mode."));
+				return false;
+#endif //WITH_EDITOR
+			}
+		}
+		// Nodes that can be replaced for inlined functions and do not need to be executed
+		if (NodeProto.OperatorType == TEXT("Constant"))
+		{
+			// Sanity check
+			if (NodeProto.Output.Num() != 1)
+			{
+				UE_LOG(LogNeuralNetworkInference, Warning, TEXT("FNeuralNetworkFromONNXTranslator::CreateOperatorsAndEditTensorArray(): Constant operator requires exactly 1 output, not %d."), NodeProto.Output.Num());
+				return false;
+			}
+			const FString& OutputTensorName = NodeProto.Output[0];
+			// Case 1: bConvertingFromSerializedData
+			if (bConvertingFromSerializedData)
+			{
+				// Sanity check
+				if (!InOutNameIndexMap.Find(OutputTensorName))
+				{
+					UE_LOG(LogNeuralNetworkInference, Warning, TEXT("FNeuralNetworkFromONNXTranslator::CreateOperatorsAndEditTensorArray()-Constant: Tensor %s was not found."), *OutputTensorName);
+					return false;
+				}
+				// Do nothing else
+			}
+			// Case 2: !bConvertingFromSerializedData
+			else
+			{
+				// Sanity check
+				if (InOutNameIndexMap.Find(OutputTensorName))
+				{
+					UE_LOG(LogNeuralNetworkInference, Warning, TEXT("FNeuralNetworkFromONNXTranslator::CreateOperatorsAndEditTensorArray()-Constant: Tensor %s was already defined and it should have not."), *OutputTensorName);
+					return false;
+				}
+				// Is it also an absolute output of the network?
+				const int64 TensorIndex = InOutTensors.Num();
+				if (InOutputNameDummyIndexMap.Find(OutputTensorName))
+				{
+					OutputNameIndexMap.Add(OutputTensorName, TensorIndex);
+				}
+				// Create tensor
+				InOutNameIndexMap.Add(OutputTensorName, TensorIndex);
+				InOutTensors.Push(FNeuralTensor(ENeuralDataType::None, TArray<int64>({}), OutputTensorName, ENeuralTensorTypeGPU::Weight));
+
+				// Create constant operator layer and get final tensor value
+				FConstantOperator ConstantOperator(NodeProto);
+				ConstantOperator.SetOutputTensors({&InOutTensors.Last()});
+				// Sanity checks
+				if (!ConstantOperator.ConfigureOutputAndInternalVariablesAndSanityChecks())
+				{
+					UE_LOG(LogNeuralNetworkInference, Warning, TEXT("FNeuralNetworkFromONNXTranslator::CreateOperatorsAndEditTensorArray()-Constant: Tensor %s was already defined and it should have not."), *OutputTensorName);
+					return false;
+				}
+				// Get output
+				ConstantOperator.ForwardCPU();
+			}
+		}
+		// Nodes/Operators that must be created and used on each Run() pass of the network
+		else
+		{
+			// Create and configure operator
+			TSharedPtr<FNeuralOperator> Operator = CreateOperator(NodeProto, OperatorInputTensorNames, InInputNameDummyIndexMap, InOutputNameDummyIndexMap, InGraphProto);
+			// ConvTranspose requires flipping OperatorInputTensors[1], so we expect for now a Weight type, so it can be flipped beforehand
+			if (!bConvertingFromSerializedData && NodeProto.OperatorType == TEXT("ConvTranspose")) //if (NodeProto.OperatorType.Contains(TEXT("Conv"), ESearchCase::CaseSensitive))
+			{
+				UE_LOG(LogNeuralNetworkInference, Display, TEXT("Flipping FConvTranspose weights. This is a temporary measurement but it will keep the weights flipped for now."));
+				// Sanity check
+				if (OperatorInputTensors.Num() > 1 && OperatorInputTensors[1]->GetTensorTypeGPU() != ENeuralTensorTypeGPU::Weight)
+				{
+					UE_LOG(LogNeuralNetworkInference, Warning,
+						TEXT("FNeuralNetworkFromONNXTranslator::CreateOperatorsAndEditTensorArray(): For now, %s operators require a constant InputTensor[0] (i.e., a fixed weight tensor), it cannot change."), *NodeProto.OperatorType);
+					return false;
+				}
+				// Flip (invert) OperatorInputTensors[1]
+				if (!OperatorInputTensors[1]->Flip(2, OperatorInputTensors[1]->GetNumberDimensions()))
+				{
+					UE_LOG(LogNeuralNetworkInference, Warning, TEXT("FNeuralNetworkFromONNXTranslator::CreateOperatorsAndEditTensorArray(): Tensor::Flip() failed."));
+					return false;
+				}
+			}
+			// Set input tensors in Operator
+			const int32 InlinedTensorIndex = Operator->SetInputTensorsAndGetInlinedIndex(OperatorInputTensors);
+			// Read/set output tensors
+			// 2 different cases for OperatorOutputTensors:
+				// Case 1/2: bConvertingFromSerializedData --> All FNeuralTensors already exist
+					// Subcase 1/2: InlinedTensorIndex > -1 --> Do nothing
+					// Subcase 2/2: InlinedTensorIndex < 0 --> Fill OperatorOutputTensors from known index
+				// Case 2/2: !bConvertingFromSerializedData
+					// Subcase 1/3: ExistingTensorIndex is found --> FNeuralTensor already exists
+					// Subcase 2/3: InlinedTensorIndex < 0 --> Output tensor must be created because operator creates its own output
+					// Subcase 3/3: InlinedTensorIndex > -1 --> Output tensor is the same than 1 of the input tensors of that operator (inlined operator)
+			if (InlinedTensorIndex > -1 && NodeProto.Output.Num() != 1)
+			{
+				UE_LOG(LogNeuralNetworkInference, Warning, TEXT("FNeuralNetworkFromONNXTranslator::CreateOperatorsAndEditTensorArray(): There cannot be an inlined tensor (index %d)"
+					" if NodeProto.Output.Num() > 1 (currently %d)."), InlinedTensorIndex, NodeProto.Output.Num());
+				return false;
+			}
+			TArray<FNeuralTensor*> OperatorOutputTensors;
+			for (const FString& OutputTensorName : NodeProto.Output)
+			{
+				// Case 1: bConvertingFromSerializedData
+				if (bConvertingFromSerializedData)
+				{
+					// Subcase 1/2: Do nothing
+					// Subcase 2/2: Fill OperatorOutputTensors from known index
+					if (InlinedTensorIndex < 0)
+					{
+						const int32 ExistingTensorIndex = InOutNameIndexMap.FindChecked(OutputTensorName);
+						OperatorOutputTensors.Push(&InOutTensors[ExistingTensorIndex]);
+					}
+				}
+				// Case 2: !bConvertingFromSerializedData
+				else
+				{
+					// Subcase 1/3: FNeuralTensor already exists (i.e., absolute output of the UNeuralNetworkLegacy or bConvertingFromSerializedData == true)
+					if (const int32* const ExistingTensorIndex = InOutNameIndexMap.Find(OutputTensorName))
+					{
+						if (InlinedTensorIndex > -1)
+						{
+							UE_LOG(LogNeuralNetworkInference, Warning,
+								TEXT("FNeuralNetworkFromONNXTranslator::CreateOperatorsAndEditTensorArray(): InlinedTensorIndex = %d for this operator (%s) but the tensor (%s) already exists."), InlinedTensorIndex, *NodeProto.OperatorType, *OutputTensorName);
+							return false;
+						}
+						OperatorOutputTensors.Push(&InOutTensors[*ExistingTensorIndex]);
+						// Is it also an absolute output of the network?
+						if (InOutputNameDummyIndexMap.Find(OutputTensorName))
+						{
+							OutputNameIndexMap.Add(OutputTensorName, *ExistingTensorIndex);
+							InOutTensors[*ExistingTensorIndex].SetTensorTypeGPU(ENeuralTensorTypeGPU::Output);
+						}
+					}
+					// Subcases 2 and 3: FNeuralTensor not found, it must be defined now (i.e., intermediate tensor of the UNeuralNetworkLegacy)
+					// Subcase 2/3: Output tensor must be created because operator creates its own output
+					else if (InlinedTensorIndex < 0)
+					{
+						// Is it also an absolute output of the network?
+						const int64 TensorIndex = InOutTensors.Num();
+						ENeuralTensorTypeGPU TensorTypeGPU = ENeuralTensorTypeGPU::IntermediateNotInitialized;
+						if (InOutputNameDummyIndexMap.Find(OutputTensorName))
+						{
+							OutputNameIndexMap.Add(OutputTensorName, TensorIndex);
+							TensorTypeGPU = ENeuralTensorTypeGPU::Output;
+						}
+						// Create tensor
+						InOutNameIndexMap.Add(OutputTensorName, TensorIndex);
+						InOutTensors.Push(FNeuralTensor(ENeuralDataType::Float, TArray<int64>({}), OutputTensorName, TensorTypeGPU));
+						OperatorOutputTensors.Push(&InOutTensors.Last());
+					}
+					// Subcase 3/3: Output tensor is the same than 1 of the input tensors of that operator (inlined operator)
+					else
+					{
+						if (NodeProto.Output.Num() != 1)
+						{
+							UE_LOG(LogNeuralNetworkInference, Warning, TEXT("FNeuralNetworkFromONNXTranslator::CreateOperatorsAndEditTensorArray()-%s-%s:"
+								" NodeProto.Output.Num() = %d != 1. Only implemented for inlined operators with 1 output. Potential issues for operators like PReLU."
+								" Although it should work as long as we pre-define the desired order of FNeuralTensors inside that operator."),
+								*NodeProto.OperatorType, *Operator->GetName(), OperatorInputTensors.Num(), NodeProto.Output.Num());
+							return false;
+						}
+						const FString& NodeProtoInputName = NodeProto.Input[InlinedTensorIndex];
+						const int32 InlinedTensorGlobalIndex = InOutNameIndexMap.FindChecked(NodeProtoInputName);
+						InOutNameIndexMap.Add(OutputTensorName, InlinedTensorGlobalIndex);
+						// Is it also an absolute output of the network?
+						if (InOutputNameDummyIndexMap.Find(OutputTensorName))
+						{
+							OutputNameIndexMap.Add(OutputTensorName, InlinedTensorGlobalIndex);
+							InOutTensors[InlinedTensorGlobalIndex].SetTensorTypeGPU(ENeuralTensorTypeGPU::Output);
+						}
+					}
+				}
+			}
+			// Set output tensors in Operator (if operator is not fully inlined)
+			if (OperatorOutputTensors.Num() > 0)
+			{
+				Operator->SetOutputTensors(OperatorOutputTensors);
+			}
+			// Create and add auxiliary tensors for that operator (if any)
+			const int64 NumberAuxiliaryTensors = Operator->GetNumberAuxiliaryTensors();
+			if (NumberAuxiliaryTensors > 0)
+			{
+				TArray<FNeuralTensor*> AuxiliaryOperatorTensors;
+				for (int64 AuxiliaryTensorIndex = 0; AuxiliaryTensorIndex < Operator->GetNumberAuxiliaryTensors(); ++AuxiliaryTensorIndex)
+				{
+					const FString AuxiliaryTensorName = NodeProto.Name + TEXT("-") + Operator->GetName() + TEXT("-Aux") + FString::FromInt(AuxiliaryTensorIndex);
+					// Uasset file
+					if (bConvertingFromSerializedData)
+					{
+						// Find and push auxiliary tensor
+						if (const int32* const AuxiliaryTensorAbsoluteIndex = InOutNameIndexMap.Find(AuxiliaryTensorName))
+						{
+							AuxiliaryOperatorTensors.Push(&InOutTensors[*AuxiliaryTensorAbsoluteIndex]);
+						}
+						// Something went wrong, the tensor should exist
+						else
+						{
+							UE_LOG(LogNeuralNetworkInference, Warning, TEXT("FNeuralNetworkFromONNXTranslator::CreateOperatorsAndEditTensorArray(): AuxiliaryTensorName %s did not exist in InOutTensors."), *AuxiliaryTensorName);
+							return false;
+						}
+					}
+					// ONNX file
+					else
+					{
+						// Sanity check
+						if (InOutNameIndexMap.Find(AuxiliaryTensorName))
+						{
+							UE_LOG(LogNeuralNetworkInference, Warning, TEXT("FNeuralNetworkFromONNXTranslator::CreateOperatorsAndEditTensorArray(): InOutNameIndexMap already has %s. Notify us of this error."), *AuxiliaryTensorName);
+							return false;
+						}
+						// Create and push new auxiliary tensor
+						InOutNameIndexMap.Add(AuxiliaryTensorName, InOutTensors.Num());
+						InOutTensors.Push(FNeuralTensor(ENeuralDataType::None, TArray<int64>({}), AuxiliaryTensorName, ENeuralTensorTypeGPU::Generic)); // ENeuralTensorTypeGPU can be modify by FNeuralOperator::Configure()
+						AuxiliaryOperatorTensors.Push(&InOutTensors.Last());
+					}
+				}
+				Operator->SetAuxiliaryTensors(AuxiliaryOperatorTensors);
+			}
+			// If operator is inlined, then input and output should already match
+			// Configure operator
+			if (!Operator->Configure())
+			{
+				UE_LOG(LogNeuralNetworkInference, Warning,
+					TEXT("FNeuralNetworkFromONNXTranslator::CreateOperatorsAndEditTensorArray(): Operator %s could not be configured, so network could not be loaded."), *Operator->GetName());
+				return false;
+			}
+			// Push Operator into OutOperators
+			OutOperators.Push(Operator);
+		}
+	}
+	// Sanity checks for bConvertingFromSerializedData
+	if (bConvertingFromSerializedData)
+	{
+		if (OutputNameIndexMap.Num() > 0)
+		{
+			UE_LOG(LogNeuralNetworkInference, Warning,
+				TEXT("FNeuralNetworkFromONNXTranslator::CreateOperatorsAndEditTensorArray(): Error in this function, OutputNameIndexMap.Num() should be 0 and not %d when reading from serialized data."),
+				OutputNameIndexMap.Num());
+			return false;
+		}
+		else if (InOutTensorsNumInit != InOutTensors.Num())
+		{
+			UE_LOG(LogNeuralNetworkInference, Warning,
+				TEXT("FNeuralNetworkFromONNXTranslator::CreateOperatorsAndEditTensorArray(): Error in this function, InOutTensors.Num() should not change from %d to %d when reading from serialized data."),
+				InOutTensorsNumInit, InOutTensors.Num());
+			return false;
+		}
+		else if (InOutNameIndexMapNumInit != InOutNameIndexMap.Num())
+		{
+			UE_LOG(LogNeuralNetworkInference, Warning,
+				TEXT("FNeuralNetworkFromONNXTranslator::CreateOperatorsAndEditTensorArray(): Error in this function, InOutNameIndexMap.Num() should not change from %d to %d when reading from serialized data."),
+				InOutTensorsNumInit, InOutTensors.Num());
+			return false;
+		}
+	}
+	// Sanity check for !bConvertingFromSerializedData
+	else if (OutputNameIndexMap.Num() != InOutputNameDummyIndexMap.Num())
+	{
+		UE_LOG(LogNeuralNetworkInference, Warning,
+			TEXT("FNeuralNetworkFromONNXTranslator::CreateOperatorsAndEditTensorArray(): OutputNameIndexMap.Num() == InOutputNameDummyIndexMap.Num() failed, %d != %d."
+				" This means that not all of the outputs defined on the ONNX file were actually needed. The program will proceed with only the required ones."),
+			OutputNameIndexMap.Num(), InOutputNameDummyIndexMap.Num());
+	}
+	// Sanity check for all cases
+	if (InOutNameIndexMap.Num() < InOutTensors.Num())
+	{
+		UE_LOG(LogNeuralNetworkInference, Warning,
+			TEXT("FNeuralNetworkFromONNXTranslator::CreateOperatorsAndEditTensorArray(): Error in this function, InOutNameIndexMap.Num() >= InOutTensors.Num() failed, %d != %d. They"
+				" should match if no inlining occurs, or more generically, InOutNameIndexMap.Num() = InOutTensors.Num() + number_inlined_tensors."),
+			InOutNameIndexMap.Num(), InOutTensors.Num());
+	}
+
+	return (OutOperators.Num() > 0);
+}
+
+bool FNeuralNetworkFromONNXTranslator::CanTensorBeInlined(const FString& InTensorName, const TMap<FString, int32>& InInputNameDummyIndexMap,
+	const TMap<FString, int32>& InOutputNameDummyIndexMap, const FNodeProto& InNodeProto, const FGraphProto& InGraphProto, const bool bCanOutputPropertiesChangeWithPostForward)
+{
+	// If tensor is one of the input/output tensors --> Operator cannot be inlined
+	if (InInputNameDummyIndexMap.Find(InTensorName) || InOutputNameDummyIndexMap.Find(InTensorName))
+	{
+		return false;
+	}
+	// If PostForward() changes the properties of the output tensor (e.g., Reshape might change its size), then it cannot be combined with a global output of the network
+	else if (bCanOutputPropertiesChangeWithPostForward)
+	{
+		for (const FValueInfoProto& GraphProtoOutput : InGraphProto.Output)
+		{
+			if (GraphProtoOutput.Name == InNodeProto.Output[0])
+			{
+				return false;
+			}
+		}
+	}
+	// If tensor is used again --> Operator cannot be inlined
+	else
+	{
+		int32 NumberOperatorsUsingTensor = 0;
+		// Iterating over each operator
+		for (const FNodeProto& NodeProto : InGraphProto.Node)
+		{
+			// Iterating over each input tensor name for that operator
+			for (const FString& InputTensorName : NodeProto.Input)
+			{
+				// If names matches, increase NumberOperatorsUsingTensor
+				if (InputTensorName == InTensorName)
+				{
+					++NumberOperatorsUsingTensor;
+					if (NumberOperatorsUsingTensor > 1)
+					{
+						return false;
+					}
+				}
+			}
+		}
+		// Sanity check
+		if (NumberOperatorsUsingTensor != 1)
+		{
+			UE_LOG(LogNeuralNetworkInference, Warning,
+				TEXT("FNeuralNetworkFromONNXTranslator::CanTensorBeInlined(): At this point, NumberOperatorsUsingTensor should be exactly 1 (not %d)."),
+				NumberOperatorsUsingTensor);
+		}
+	}
+	// None of the above is true --> Operator can be safely inlined
+	return true;
+}
+
+TSet<uint32> FNeuralNetworkFromONNXTranslator::GetPotentiallyInlinedTensors(const TArray<FString>& InTensorNames, const TMap<FString, int32>& InInputNameDummyIndexMap,
+	const TMap<FString, int32>& InOutputNameDummyIndexMap, const FNodeProto& InNodeProto, const FGraphProto& InGraphProto)
+{
+	TSet<uint32> PotentiallyInlinedTensors;
+	for (int32 TensorNameIndex = 0; TensorNameIndex < InTensorNames.Num(); ++TensorNameIndex)
+	{
+		const FString& TensorName = InTensorNames[TensorNameIndex];
+		if (FNeuralNetworkFromONNXTranslator::CanTensorBeInlined(TensorName, InInputNameDummyIndexMap, InOutputNameDummyIndexMap, InNodeProto, InGraphProto))
+		{
+			PotentiallyInlinedTensors.Add((uint32)TensorNameIndex);
+		}
+	}
+	return PotentiallyInlinedTensors;
+}
+
+// Eg Element-wise operators with 0 attributes
+#define IF_MATCHED_INLINABLE_OPERATOR_INITIALIZATION(OperatorTypeSuffix) \
+	if (InNodeProto.OperatorType == TEXT(#OperatorTypeSuffix)) /* Eg TEXT("Relu") */ \
+	{ \
+		const bool bCanTensorBeInlined = CanTensorBeInlined(InInputTensorNamesForOperator[0], InInputNameDummyIndexMap, InOutputNameDummyIndexMap, InNodeProto, InGraphProto); \
+		Operator = MakeShared<F##OperatorTypeSuffix##Operator>(bCanTensorBeInlined); /* Eg FReluOperator */ \
+	}
+
+// Eg Element-wise operators with n attributes or operators like BatchNormalization
+#define IF_MATCHED_INLINABLE_OPERATOR_WITH_PROTO_INITIALIZATION(OperatorTypeSuffix, bCanOutputPropertiesChangeWithPostForward) \
+	if (InNodeProto.OperatorType == TEXT(#OperatorTypeSuffix)) /* Eg TEXT("BatchNormalization") */ \
+	{ \
+		const bool bCanTensorBeInlined = CanTensorBeInlined(InInputTensorNamesForOperator[0], InInputNameDummyIndexMap, InOutputNameDummyIndexMap, InNodeProto, InGraphProto, bCanOutputPropertiesChangeWithPostForward); \
+		Operator = MakeShared<F##OperatorTypeSuffix##Operator>(bCanTensorBeInlined, InNodeProto); /* Eg FBatchNormalizationOperator */ \
+	}
+
+// Eg Conv(Transpose) and Gemm
+#define IF_MATCHED_OPERATOR_WITH_PROTO_INITIALIZATION(OperatorTypeSuffix) \
+	if (InNodeProto.OperatorType == TEXT(#OperatorTypeSuffix)) /* Eg TEXT("Conv") */ \
+	{ \
+		Operator = MakeShared<F##OperatorTypeSuffix##Operator>(InNodeProto); /* Eg FConvOperator */ \
+	}
+
+#define IF_MATCHED_MULTIDIRECTIONAL_BROADCASTING_OPERATOR_INITIALIZATION(OperatorTypeSuffix) \
+	if (InNodeProto.OperatorType == TEXT(#OperatorTypeSuffix)) /* Eg TEXT("Add") */ \
+	{ \
+		const TSet<uint32> PotentiallyInlinedTensors = GetPotentiallyInlinedTensors(InInputTensorNamesForOperator, InInputNameDummyIndexMap, InOutputNameDummyIndexMap, InNodeProto, InGraphProto); \
+		Operator = MakeShared<F##OperatorTypeSuffix##Operator>(PotentiallyInlinedTensors); /* Eg FAddOperator */ \
+	}
+
+TSharedPtr<FNeuralOperator> FNeuralNetworkFromONNXTranslator::CreateOperator(const FNodeProto& InNodeProto, const TArray<FString>& InInputTensorNamesForOperator,
+	const TMap<FString, int32>& InInputNameDummyIndexMap, const TMap<FString, int32>& InOutputNameDummyIndexMap, const FGraphProto& InGraphProto)
+{
+	TSharedPtr<FNeuralOperator> Operator;
+	IF_MATCHED_INLINABLE_OPERATOR_INITIALIZATION(Abs)
+	else IF_MATCHED_INLINABLE_OPERATOR_INITIALIZATION(Acos)
+	else IF_MATCHED_MULTIDIRECTIONAL_BROADCASTING_OPERATOR_INITIALIZATION(Add)
+	else IF_MATCHED_INLINABLE_OPERATOR_INITIALIZATION(Asin)
+	else IF_MATCHED_INLINABLE_OPERATOR_INITIALIZATION(Atan)
+	else IF_MATCHED_INLINABLE_OPERATOR_WITH_PROTO_INITIALIZATION(BatchNormalization, false)
+	else IF_MATCHED_INLINABLE_OPERATOR_INITIALIZATION(Ceil)
+	else IF_MATCHED_OPERATOR_WITH_PROTO_INITIALIZATION(Conv)
+	else IF_MATCHED_OPERATOR_WITH_PROTO_INITIALIZATION(ConvTranspose)
+	if (InNodeProto.OperatorType == TEXT("ConvTranspose"))
+	{
+		TSharedPtr<FConvTransposeOperator> ConvTransposeOperator;
+		ConvTransposeOperator = MakeShared<FConvTransposeOperator>(InNodeProto);
+		ConvTransposeOperator->SetWhetherWeightTensorWasFlipped(true);
+		Operator = MoveTemp(ConvTransposeOperator);
+	}
+	else IF_MATCHED_INLINABLE_OPERATOR_INITIALIZATION(Cos)
+	else IF_MATCHED_INLINABLE_OPERATOR_INITIALIZATION(Cosh)
+	else IF_MATCHED_MULTIDIRECTIONAL_BROADCASTING_OPERATOR_INITIALIZATION(Div)
+	else IF_MATCHED_INLINABLE_OPERATOR_INITIALIZATION(Exp)
+	else IF_MATCHED_INLINABLE_OPERATOR_INITIALIZATION(Floor)
+	else IF_MATCHED_OPERATOR_WITH_PROTO_INITIALIZATION(Gemm)
+	else IF_MATCHED_INLINABLE_OPERATOR_WITH_PROTO_INITIALIZATION(LeakyRelu, false)
+	else IF_MATCHED_INLINABLE_OPERATOR_INITIALIZATION(Log)
+	else IF_MATCHED_MULTIDIRECTIONAL_BROADCASTING_OPERATOR_INITIALIZATION(Mul)
+	else IF_MATCHED_INLINABLE_OPERATOR_INITIALIZATION(Neg)
+	else IF_MATCHED_MULTIDIRECTIONAL_BROADCASTING_OPERATOR_INITIALIZATION(Pow)
+	else IF_MATCHED_INLINABLE_OPERATOR_INITIALIZATION(Reciprocal)
+	else IF_MATCHED_INLINABLE_OPERATOR_INITIALIZATION(Relu)
+	else IF_MATCHED_INLINABLE_OPERATOR_WITH_PROTO_INITIALIZATION(Reshape, true)
+	else IF_MATCHED_INLINABLE_OPERATOR_INITIALIZATION(Round)
+	else IF_MATCHED_INLINABLE_OPERATOR_INITIALIZATION(Sigmoid)
+	else IF_MATCHED_INLINABLE_OPERATOR_INITIALIZATION(Sign)
+	else IF_MATCHED_INLINABLE_OPERATOR_INITIALIZATION(Sin)
+	else IF_MATCHED_INLINABLE_OPERATOR_INITIALIZATION(Sinh)
+	else IF_MATCHED_INLINABLE_OPERATOR_INITIALIZATION(Sqrt)
+	else IF_MATCHED_MULTIDIRECTIONAL_BROADCASTING_OPERATOR_INITIALIZATION(Sub)
+	else IF_MATCHED_INLINABLE_OPERATOR_INITIALIZATION(Tan)
+	else IF_MATCHED_INLINABLE_OPERATOR_INITIALIZATION(Tanh)
+	if (!Operator.IsValid())
+	{
+		UE_LOG(LogNeuralNetworkInference, Warning, TEXT("FNeuralNetworkFromONNXTranslator::CreateOperator(): Unknown or unimplemented operator %s."), *InNodeProto.OperatorType);
+	}
+	return Operator;
+}
diff --git a/Engine/Plugins/Experimental/NeuralNetworkInference/Source/NeuralNetworkInference/Private/NeuralNetworkFromONNXTranslator.h b/Engine/Plugins/Experimental/NeuralNetworkInference/Source/NeuralNetworkInference/Private/NeuralNetworkFromONNXTranslator.h
new file mode 100644
--- /dev/null
+++ b/Engine/Plugins/Experimental/NeuralNetworkInference/Source/NeuralNetworkInference/Private/NeuralNetworkFromONNXTranslator.h
@@ -0,0 +1,65 @@
+// Copyright Epic Games, Inc. All Rights Reserved.
+
+#pragma once
+
+#include "CoreMinimal.h"
+#include "ModelProto.h"
+#include "NeuralOperator.h"
+#include "NeuralTensor.h"
+#include "NeuralTensorManager.h"
+
+class FNeuralNetworkFromONNXTranslator
+{
+public:
+	/**
+	 * It takes a FGraphProto and model path as inputs, and produces an array of FNeuralOperator's and FNeuralTensorManager as result.
+	 * These 2 outputs are used by UNeuralNetworkLegacy to create and load a (deep) neural network architecture and its weights.
+	 * @param InTensorManager It is const, but kept not const for CreateOperatorsAndEditTensorArray() compatibility.
+	 *
+	 * 2 alternatives:
+	 * Without InModelPath, InOutTensorManager is assumed to be a const input variable (assuming the tensors are all already loaded, e.g., on de-serialization).
+	 * With InModelPath (Editor-only function), InOutTensorManager is assumed to be an output variable that will be filled based on InModelPath.
+	 */
+	static bool Translate(TArray<TSharedPtr<FNeuralOperator>>& OutOperators, FNeuralTensorManager& InTensorManager, const FGraphProto& InGraphProto);
+#if WITH_EDITOR
+	static bool Translate(TArray<TSharedPtr<FNeuralOperator>>& OutOperators, FNeuralTensorManager& OutTensorManager, const FGraphProto& InGraphProto, const FString& InModelPath);
+#endif //WITH_EDITOR
+
+private:
+	/**
+	 * Auxiliary function for Translate().
+	 * It initializes both OutOperators and OutputNameIndexMap, as well as modify InOutTensors and InOutNameIndexMap.
+	 * OutputNameIndexMap is just the copy of InOutputNameDummyIndexMap while filling its int32 Value fields. Given the dual Game/Editor behaviors, and to optimize Game mode, kept as 2 different rather than modifying in-place as InOutOutputNameDummyIndexMap.
+	 *
+	 * 2 alternatives:
+	 * - Game and/or Editor (efficient): If InModelPath is empty, InOutTensorManager is assumed to be a const input variable (assuming the tensors are all already loaded, e.g., on de-serialization).
+	 * - Editor-only: If InModelPath is not empty, InOutTensorManager is assumed to be an output variable that will be filled based on InModelPath.
+	 *
+	 * @param OutputNameIndexMap It is created from InOutputNameDummyIndexMap, assuming the FString part of the mapping is right, and it will fill the index (int32 part of it).
+	 * @param InInputNameDummyIndexMap The int32 part is irrelevant, the only important part is whether the FString field is filled or not, regardless of its value. Kept as TMap for performance to integrate better with FTensorManager.
+	 * @param InOutputNameDummyIndexMap The int32 part is irrelevant, the only important part is whether the FString field is filled or not, regardless of its value. Kept as TMap for performance to integrate better with FTensorManager.
+	 */
+	static bool CreateOperatorsAndEditTensorArray(TArray<TSharedPtr<FNeuralOperator>>& OutOperators, TMap<FString, int32>& OutputNameIndexMap, TArray<FNeuralTensor>& InOutTensors,
+		TMap<FString, int32>& InOutNameIndexMap, const TMap<FString, int32>& InInputNameDummyIndexMap, const TMap<FString, int32>& InOutputNameDummyIndexMap,
+		const FGraphProto& InGraphProto, const FString& InModelPath = TEXT(""));
+
+	/**
+	 * It will return FNeuralOperator::EDataType::NotInlined in any of the 2 following cases occur:
+	 * - If the tensor is an EGPUDataType::Output.
+	 * - If the tensor is used by another layer.
+	 * It will return FNeuralOperator::EDataType::Inlined otherwise.
+	 */
+	static bool CanTensorBeInlined(const FString& InTensorName, const TMap<FString, int32>& InInputNameDummyIndexMap,
+		const TMap<FString, int32>& InOutputNameDummyIndexMap, const FNodeProto& InNodeProto, const FGraphProto& InGraphProto, const bool bCanOutputPropertiesChangeWithPostForward = false);
+	/**
+	 * It will return which input tensors can be inlined.
+	 */
+	static TSet<uint32> GetPotentiallyInlinedTensors(const TArray<FString>& InTensorNames, const TMap<FString, int32>& InInputNameDummyIndexMap,
+		const TMap<FString, int32>& InOutputNameDummyIndexMap, const FNodeProto& InNodeProto, const FGraphProto& InGraphProto);
+
+	/**
+	 * Given the input arguments, it creates the proper TShared<FNeuralOperator> child instance.
+	 */
+	static TSharedPtr<FNeuralOperator> CreateOperator(const FNodeProto& InNodeProto, const TArray<FString>& InInputTensorNamesForOperator,
+		const TMap<FString, int32>& InInputNameDummyIndexMap, const TMap<FString, int32>& InOutputNameDummyIndexMap, const FGraphProto& InGraphProto);
+};
diff --git a/Engine/Plugins/Experimental/NeuralNetworkInference/Source/NeuralNetworkInference/Private/NeuralNetworkInferenceUtils.cpp b/Engine/Plugins/Experimental/NeuralNetworkInference/Source/NeuralNetworkInference/Private/NeuralNetworkInferenceUtils.cpp
--- a/Engine/Plugins/Experimental/NeuralNetworkInference/Source/NeuralNetworkInference/Private/NeuralNetworkInferenceUtils.cpp
+++ b/Engine/Plugins/Experimental/NeuralNetworkInference/Source/NeuralNetworkInference/Private/NeuralNetworkInferenceUtils.cpp
@@ -2,6 +2,9 @@
 
 #include "NeuralNetworkInferenceUtils.h"
 #include "HAL/PlatformTime.h"
+#include "RenderingThread.h"
+#include "RHI.h"
+#include <atomic>
 
 
 
@@ -55,3 +58,18 @@ bool FNeuralNetworkInferenceUtils::SizeSanityChecks(const TArray<FNeuralTensor*>
 	}
 	return true;
 }
+
+void FNeuralNetworkInferenceUtils::WaitUntilRHIFinished()
+{
+	std::atomic<bool> bDidGPUFinish(false);
+	ENQUEUE_RENDER_COMMAND(ForwardGPU_Gemm_RenderThread)(
+		[&bDidGPUFinish](FRHICommandListImmediate& RHICmdList)
+		{
+			bDidGPUFinish = true;
+		}
+	);
+	while (!bDidGPUFinish)
+	{
+		FPlatformProcess::Sleep(0.1e-3);
+	}
+}
diff --git a/Engine/Plugins/Experimental/NeuralNetworkInference/Source/NeuralNetworkInference/Private/NeuralNetworkInferenceUtils.h b/Engine/Plugins/Experimental/NeuralNetworkInference/Source/NeuralNetworkInference/Private/NeuralNetworkInferenceUtils.h
--- a/Engine/Plugins/Experimental/NeuralNetworkInference/Source/NeuralNetworkInference/Private/NeuralNetworkInferenceUtils.h
+++ b/Engine/Plugins/Experimental/NeuralNetworkInference/Source/NeuralNetworkInference/Private/NeuralNetworkInferenceUtils.h
@@ -5,8 +5,10 @@
 #include "CoreMinimal.h"
 #include "NeuralTensor.h"
 #include "Algo/Accumulate.h"
+#include "Stats/Stats.h"
 
 DEFINE_LOG_CATEGORY_STATIC(LogNeuralNetworkInference, Display, All);
+DECLARE_STATS_GROUP(TEXT("MachineLearning"), STATGROUP_MachineLearning, STATCAT_Advanced);
 
 class FNeuralNetworkInferenceUtils
 {
@@ -29,6 +31,12 @@ public:
 
 	template <typename ValueType, typename ContainerType>
 	static ValueType Product(const ContainerType& InContainer);
+
+	/**
+	 * It blocks the current thread until the RHI thread has finished all instructions before this point.
+	 * @returns The time it waited.
+	 */
+	static void WaitUntilRHIFinished();
 };
 
 
diff --git a/Engine/Plugins/Experimental/NeuralNetworkInference/Source/NeuralNetworkInference/Private/NeuralNetworkLegacy.cpp b/Engine/Plugins/Experimental/NeuralNetworkInference/Source/NeuralNetworkInference/Private/NeuralNetworkLegacy.cpp
new file mode 100644
--- /dev/null
+++ b/Engine/Plugins/Experimental/NeuralNetworkInference/Source/NeuralNetworkInference/Private/NeuralNetworkLegacy.cpp
@@ -0,0 +1,484 @@
+// Copyright Epic Games, Inc. All Rights Reserved.
+
+#include "NeuralNetworkLegacy.h"
+#include "NeuralOperator.h"
+#include "NeuralNetworkFromONNXTranslator.h"
+#include "NeuralNetworkInferenceUtils.h"
+#include "NeuralNetworkInferenceVersion.h"
+#include "EditorFramework/AssetImportData.h"
+#include "HAL/FileManager.h"
+#include "RenderGraphBuilder.h"
+#include "RenderingThread.h"
+#include "RHI.h"
+
+#if WITH_EDITOR
+#if PLATFORM_WINDOWS
+#include "ModelProtoFileReader.h"
+#endif //PLATFORM_WINDOWS
+#endif //WITH_EDITOR
+
+
+
+/* UNeuralNetworkLegacy structors
+ *****************************************************************************/
+
+UNeuralNetworkLegacy::UNeuralNetworkLegacy()
+	: bIsLoaded(false)
+	, DeviceType(ENeuralDeviceType::GPU)
+	, bAreTensorsInGpu(false)
+{
+}
+
+UNeuralNetworkLegacy::~UNeuralNetworkLegacy()
+{
+}
+
+
+
+/* UNeuralNetworkLegacy public functions
+ *****************************************************************************/
+
+void UNeuralNetworkLegacy::PostInitProperties()
+{
+#if WITH_EDITORONLY_DATA
+	if (!HasAnyFlags(RF_ClassDefaultObject) && !AssetImportData)
+	{
+		GetAndMaybeCreateAssetImportData();
+	}
+#endif
+	Super::PostInitProperties();
+}
+
+void UNeuralNetworkLegacy::PostLoad()
+{
+	Super::PostLoad();
+	// Turn ModelProto into Operators
+	if (bIsLoaded)
+	{
+		// Sanity checks
+		if (!IsLoaded() || !TensorManager.IsLoaded() || !ModelProto.IsLoaded())
+		{
+			UE_LOG(LogNeuralNetworkInference, Warning, TEXT("UNeuralNetworkLegacy::PostLoad(): 1 of the 3 IsLoaded() failed: UNeuralNetworkLegacy.IsLoaded() = %d, TensorManager.IsLoaded() = %d, ModelProto.IsLoaded() = %d."),
+				IsLoaded(), TensorManager.IsLoaded(), ModelProto.IsLoaded());
+			bIsLoaded = false;
+		}
+		else if (!FNeuralNetworkInferenceVersion::CheckVersion(Version) || !FNeuralNetworkInferenceVersion::CheckVersion(TensorManager.GetVersion()) || !FNeuralNetworkInferenceVersion::CheckVersion(ModelProto.GetVersion()))
+		{
+			UE_LOG(LogNeuralNetworkInference, Warning, TEXT("UNeuralNetworkLegacy::PostLoad(): CheckVersion() failed."));
+			bIsLoaded = false;
+		}
+		// Turn ModelProto into Operators
+		else if (!FNeuralNetworkFromONNXTranslator::Translate(Operators, TensorManager, ModelProto.GetGraph()))
+		{
+			UE_LOG(LogNeuralNetworkInference, Warning, TEXT("UNeuralNetworkLegacy::PostLoad(): UNeuralNetworkLegacy could not be configured from its FModelProto."));
+		}
+	}
+}
+
+void UNeuralNetworkLegacy::Serialize(FArchive& Archive)
+{
+#if WITH_EDITORONLY_DATA
+	// Setup source data.
+	if (Archive.IsSaving() && Archive.IsPersistent())
+	{
+		ReimportAssetFromEditorData();
+	}
+#endif // WITH_EDITORONLY_DATA
+	Super::Serialize(Archive);
+}
+
+#if WITH_EDITOR
+void UNeuralNetworkLegacy::ReimportAssetFromEditorData()
+{
+	//Get the re-import filename
+	const FString ImportedFilename = AssetImportData->GetFirstFilename();
+	if (ImportedFilename.Len() > 0)
+	{
+		// Ensure that the file provided by the path exists
+		if (IFileManager::Get().FileSize(*ImportedFilename) != INDEX_NONE)
+		{
+			UE_LOG(LogNeuralNetworkInference, Display, TEXT("Performing atomic reimport of [%s]"), *ImportedFilename);
+			Load(ImportedFilename);
+		}
+	}
+}
+
+#if WITH_EDITOR
+UAssetImportData* UNeuralNetworkLegacy::GetAssetImportData() const
+{
+	return AssetImportData;
+}
+#endif // WITH_EDITOR
+
+UAssetImportData* UNeuralNetworkLegacy::GetAndMaybeCreateAssetImportData()
+{
+	// An existing import data object was not found, so make one here.
+	if (!AssetImportData)
+	{
+		AssetImportData = NewObject<UAssetImportData>(this, TEXT("AssetImportData"));
+	}
+	return AssetImportData;
+}
+#endif // WITH_EDITOR
+
+#if WITH_EDITOR
+bool UNeuralNetworkLegacy::Load(const FString& InFilePath)
+{
+#if PLATFORM_WINDOWS
+
+	DECLARE_SCOPE_CYCLE_COUNTER(TEXT("UNeuralNetworkLegacy_Load_File"), STAT_UNeuralNetworkLegacy_Load_File, STATGROUP_MachineLearning);
+
+	// Clean previous networks
+	if (bIsLoaded)
+	{
+		UE_LOG(LogNeuralNetworkInference, Warning, TEXT("UNeuralNetworkLegacy::Load(): A model was previously loaded, removing it and reloading the new model."));
+		Operators.Empty();
+		ModelProto = FModelProto();
+		bIsLoaded = false;
+		bAreTensorsInGpu = false;
+	}
+
+	// Read ModelProto
+	if (!FModelProtoFileReader::ReadModelProtoFromFile(ModelProto, InFilePath) || !ModelProto.IsLoaded() || !FNeuralNetworkInferenceVersion::CheckVersion(ModelProto.GetVersion()))
+	{
+		UE_LOG(LogNeuralNetworkInference, Warning, TEXT("UNeuralNetworkLegacy::Load(): Model could not be loaded from %s or is outdated. IsLoaded() = %d."),
+			*InFilePath, ModelProto.IsLoaded());
+		return false;
+	}
+	// Turn ModelProto into Operators
+	bIsLoaded = FNeuralNetworkFromONNXTranslator::Translate(Operators, TensorManager, ModelProto.GetGraph(), InFilePath);
+	if (!TensorManager.IsLoaded() || !FNeuralNetworkInferenceVersion::CheckVersion(TensorManager.GetVersion()))
+	{
+		UE_LOG(LogNeuralNetworkInference, Warning, TEXT("UNeuralNetworkLegacy::Load(): TensorManager could not be loaded from %s or is outdated. IsLoaded() = %d."),
+			*InFilePath, TensorManager.IsLoaded());
+		return false;
+	}
+	// Update Version
+	if (bIsLoaded)
+	{
+		Version = FNeuralNetworkInferenceVersion::GetVersion();
+	}
+	return bIsLoaded;
+
+#else //PLATFORM_WINDOWS
+	UE_LOG(LogNeuralNetworkInference, Warning, TEXT("UNeuralNetworkLegacy::Load(): Only implemented for Windows."),
+		*InFilePath, ModelProto.IsLoaded());
+	return false;
+#endif //PLATFORM_WINDOWS
+}
+#endif //WITH_EDITOR
+
+bool UNeuralNetworkLegacy::Load(FNeuralTensorManager& InTensorManager, const TArray<TSharedPtr<FNeuralOperator>>& InOperators)
+{
+	DECLARE_SCOPE_CYCLE_COUNTER(TEXT("UNeuralNetworkLegacy_Load"), STAT_UNeuralNetworkLegacy_Load, STATGROUP_MachineLearning);
+	if (!InTensorManager.IsLoaded() || !FNeuralNetworkInferenceVersion::CheckVersion(InTensorManager.GetVersion()))
+	{
+		UE_LOG(LogNeuralNetworkInference, Warning, TEXT("UNeuralNetworkLegacy::Load(): TensorManager could not be loaded or is outdated. IsLoaded() = %d."),
+			InTensorManager.IsLoaded());
+	}
+	Swap(TensorManager, InTensorManager);
+	Operators = InOperators;
+	bIsLoaded = (Operators.Num() > 0 && TensorManager.IsLoaded());
+	// Update Version
+	if (bIsLoaded)
+	{
+		Version = FNeuralNetworkInferenceVersion::GetVersion();
+	}
+	return bIsLoaded;
+}
+
+bool UNeuralNetworkLegacy::IsLoaded() const
+{
+	return bIsLoaded;
+}
+
+ENeuralDeviceType UNeuralNetworkLegacy::GetDeviceType() const
+{
+	return DeviceType;
+}
+
+void UNeuralNetworkLegacy::SetDeviceType(const ENeuralDeviceType InDeviceType)
+{
+	DeviceType = InDeviceType;
+}
+
+UNeuralNetworkLegacy::FOnAsyncRunCompletedInAnyThread& UNeuralNetworkLegacy::GetOnAsyncRunCompletedInAnyThreadDelegate()
+{
+	return OnAsyncRunCompletedInAnyThreadDelegate;
+}
+
+const TArray<FNeuralTensor>& UNeuralNetworkLegacy::GetTensors() const
+{
+	return TensorManager.GetTensors();
+}
+
+void UNeuralNetworkLegacy::SetInputFromTensorCopy(const FNeuralTensor& InTensorMap)
+{
+	return TensorManager.SetInputFromTensorCopy(InTensorMap);
+}
+
+void UNeuralNetworkLegacy::SetInputFromTensorMapCopy(const TMap<FString, FNeuralTensor>& InTensorMap)
+{
+	return TensorManager.SetInputFromTensorMapCopy(InTensorMap);
+}
+
+TMap<FString, void*> UNeuralNetworkLegacy::CreateInputDataPointersMutable()
+{
+	return TensorManager.CreateInputDataPointersMutable();
+}
+
+FRDGBufferUAVRef UNeuralNetworkLegacy::GetInputBufferUAVRef()
+{
+	return TensorManager.GetInputBufferUAVRef();
+}
+
+TMap<FString, FRDGBufferUAVRef> UNeuralNetworkLegacy::CreateInputBufferUAVRefs()
+{
+	return TensorManager.CreateInputBufferUAVRefs();
+}
+
+const FNeuralTensor& UNeuralNetworkLegacy::GetInputTensor() const
+{
+	return TensorManager.GetInputTensor();
+}
+
+const TMap<FString, int32>& UNeuralNetworkLegacy::GetInputNameIndexMap() const
+{
+	return TensorManager.GetInputNameIndexMap();
+}
+
+const FNeuralTensor& UNeuralNetworkLegacy::GetOutputTensor() const
+{
+	return TensorManager.GetOutputTensor();
+}
+
+const TMap<FString, int32>& UNeuralNetworkLegacy::GetOutputNameIndexMap() const
+{
+	return TensorManager.GetOutputNameIndexMap();
+}
+
+const FRDGBufferSRVRef UNeuralNetworkLegacy::GetOutputBufferSRVRef() const
+{
+	return TensorManager.GetOutputBufferSRVRef();
+}
+
+TMap<FString, const FRDGBufferSRVRef> UNeuralNetworkLegacy::CreateOutputBufferSRVRefs() const
+{
+	return TensorManager.CreateOutputBufferSRVRefs();
+}
+
+TMap<FString, FNeuralTensor> UNeuralNetworkLegacy::CreateInputTensorMap() const
+{
+	return TensorManager.CreateInputTensorMap();
+}
+
+TMap<FString, FNeuralTensor> UNeuralNetworkLegacy::CreateOutputTensorMap() const
+{
+	return TensorManager.CreateOutputTensorMap();
+}
+
+void UNeuralNetworkLegacy::Run(const EGPUSynchronousMode InGPUSynchronousMode, const ENeuralDeviceType InInputDeviceType, const ENeuralDeviceType InOutputDeviceType, const bool bRunGPUEmptyOnlyForProfiling)
+{
+	DECLARE_SCOPE_CYCLE_COUNTER(TEXT("UNeuralNetworkLegacy_Run"), STAT_UNeuralNetworkLegacy_Run, STATGROUP_MachineLearning);
+	// Sanity checks
+	if (!bIsLoaded)
+	{
+		UE_LOG(LogNeuralNetworkInference, Warning, TEXT("UNeuralNetworkLegacy::Run(): No architecture has been loaded yet. Run() will not work until IsLoaded() returns true."));
+		return;
+	}
+	// Run UNeuralNetworkLegacy
+	if (Operators.Num() > 0)
+	{
+		// Run graph
+		if (DeviceType == ENeuralDeviceType::CPU)
+		{
+			DECLARE_SCOPE_CYCLE_COUNTER(TEXT("UNeuralNetworkLegacy_Run::Forward_CPU"), STAT_UNeuralNetworkLegacy_Run_Forward_CPU, STATGROUP_MachineLearning);
+			// Run each operator forward pass
+			for (TSharedPtr<FNeuralOperator>& Operator : Operators)
+			{
+				Operator->ForwardCPU();
+			}
+			// Run each operator post forward pass
+			for (TSharedPtr<FNeuralOperator>& Operator : Operators)
+			{
+				Operator->PostForwardCPU();
+			}
+		}
+		else if (DeviceType == ENeuralDeviceType::GPU)
+		{
+			DECLARE_SCOPE_CYCLE_COUNTER(TEXT("UNeuralNetworkLegacy_Run::Forward_GPU"), STAT_UNeuralNetworkLegacy_Run_Forward_GPU, STATGROUP_MachineLearning);
+			// Sanity check
+			if (TensorManager.GetTensorsMutable().Num() < 1)
+			{
+				UE_LOG(LogNeuralNetworkInference, Warning, TEXT("UNeuralNetworkLegacy::Run(): Tensors.Num() = %d (should be > 0)."), TensorManager.GetTensorsMutable().Num());
+				return;
+			}
+
+			// On RHI thread
+			ENQUEUE_RENDER_COMMAND(UNeuralNetworkLegacy_Run_RenderThread)(
+				[this, bRunGPUEmptyOnlyForProfiling, InGPUSynchronousMode, InInputDeviceType, InOutputDeviceType](FRHICommandListImmediate& RHICmdList)
+				{
+					FMemMark Mark(FMemStack::Get());
+					FRDGBuilder GraphBuilder(RHICmdList, RDG_EVENT_NAME("UNeuralNetworkLegacy::Run()"));
+
+					// Move memory from CPU to GPU
+					TArray<FNeuralTensor>& Tensors = TensorManager.GetTensorsMutable();
+					const bool bIsInputInCPU = (InInputDeviceType == ENeuralDeviceType::CPU);
+					// Move only input tensors to GPU (once per UNeuralNetworkLegacy::Run())
+					if (bAreTensorsInGpu)
+					{
+						// If input is on CPU, move inputs to GPU
+						if (bIsInputInCPU)
+						{
+							for (const int32 InputIndex : TensorManager.GetInputIndexes())
+							{
+								Tensors[InputIndex].ToGPU_RenderThread(&GraphBuilder);
+							}
+							for (const int32 NonInputIndex : TensorManager.GetNonInputIndexes())
+							{
+								Tensors[NonInputIndex].UpdateSRVAndOrUAV_RenderThread(&GraphBuilder);
+							}
+						}
+						// If input is already on GPU, just refresh all w.r.t. GraphBuilder
+						else
+						{
+							for (FNeuralTensor& Tensor : Tensors)
+							{
+								Tensor.UpdateSRVAndOrUAV_RenderThread(&GraphBuilder);
+							}
+						}
+					}
+					// Move all (input, intermediate, weight, output) tensors to GPU and also call Operators.ToGPU() to move their auxiliary memory into GPU
+					// Only once per UNeuralNetworkLegacy instance, or until Load() is called again
+					else
+					{
+						// If input is on CPU, move inputs to GPU
+						if (bIsInputInCPU)
+						{
+							// Move tensors to GPU
+							for (FNeuralTensor& Tensor : Tensors)
+							{
+								Tensor.ToGPU_RenderThread(&GraphBuilder);
+							}
+						}
+						// If input is already on GPU, just refresh inputs
+						else
+						{
+							for (const int32 InputIndex : TensorManager.GetInputIndexes())
+							{
+								Tensors[InputIndex].UpdateSRVAndOrUAV_RenderThread(&GraphBuilder);
+							}
+							for (const int32 NonInputIndex : TensorManager.GetNonInputIndexes())
+							{
+								Tensors[NonInputIndex].ToGPU_RenderThread(&GraphBuilder);
+							}
+						}
+						// Operators->ToGPU()
+						for (TSharedPtr<FNeuralOperator>& Operator : Operators)
+						{
+							Operator->ToGPU_RenderThread();
+						}
+						// bAreTensorsInGpu is true
+						bAreTensorsInGpu = true;
+					}
+
+					if (!bRunGPUEmptyOnlyForProfiling)
+					{
+						// Run each operator forward pass
+						for (TSharedPtr<FNeuralOperator>& Operator : this->Operators)
+						{
+							Operator->ForwardGPU_RenderThread(&GraphBuilder);
+						}
+
+						// Run each operator post forward pass
+						for (TSharedPtr<FNeuralOperator>& Operator : this->Operators)
+						{
+							Operator->PostForwardGPU_RenderThread(&GraphBuilder);
+						}
+					}
+
+					// Move memory from GPU to CPU
+					const bool bIsOutputInCPU = (InOutputDeviceType == ENeuralDeviceType::CPU);
+					if (bIsOutputInCPU)
+					{
+						for (const int32 OutputIndex : TensorManager.GetOutputIndexes())
+						{
+							Tensors[OutputIndex].ToCPU_RenderThread(&GraphBuilder);
+						}
+					}
+
+					// Broadcast delegates (from the render thread)
+					if (InGPUSynchronousMode == EGPUSynchronousMode::Asynchronous)
+					{
+						GraphBuilder.AddPass(
+							RDG_EVENT_NAME("Async delegate broadcast"),
+							ERDGPassFlags::None,
+							[this](FRHICommandListImmediate& RHICmdList)
+						{
+							OnAsyncRunCompletedInAnyThreadDelegate.ExecuteIfBound();
+						});
+					}
+
+					// Execute Render Graph
+					GraphBuilder.Execute();
+				}
+			);
+
+			// Block thread until GPU has finished
+			if (InGPUSynchronousMode == EGPUSynchronousMode::Synchronous)
+			{
+				FNeuralNetworkInferenceUtils::WaitUntilRHIFinished();
+			}
+		}
+		else
+		{
+			UE_LOG(LogNeuralNetworkInference, Warning, TEXT("UNeuralNetworkLegacy::Run(): Unknown DeviceType = %d."), (int32)DeviceType);
+		}
+	}
+	else
+	{
+		UE_LOG(LogNeuralNetworkInference, Warning, TEXT("UNeuralNetworkLegacy::Run() called with an empty model."));
+	}
+}
+
+FString UNeuralNetworkLegacy::ToString() const
+{
+	// Add GraphProto
+	FString String = ModelProto.GetGraph().ToString();
+	// Add FNeuralTensor(s)
+	String += TEXT("TensorManager:\n");
+	const TMap<FString, int32>& NameIndexMap = TensorManager.GetNameIndexMap();
+	const TArray<FNeuralTensor>& Tensors = TensorManager.GetTensors();
+	if (NameIndexMap.Num() > 0)
+	{
+		for (const auto& NameIndexPair : NameIndexMap)
+		{
+			String += FString::Format(TEXT(" -{0}: {1}\n"), { NameIndexPair.Key, Tensors[NameIndexPair.Value].ToString(20) });
+		}
+	}
+	else
+	{
+		for (const FNeuralTensor& Tensor : Tensors)
+		{
+			String += FString::Format(TEXT(" -{0}\n"), { Tensor.ToString(20) });
+		}
+	}
+	String += TEXT("InputTensorMap:\n");
+	for (const auto& NameIndexPair : TensorManager.GetInputNameIndexMap())
+	{
+		String += FString::Format(TEXT(" -{0}: {1}\n"), { NameIndexPair.Key, Tensors[NameIndexPair.Value].ToString(20) });
+	}
+	String += TEXT("OutputTensorMap:\n");
+	for (const auto& NameIndexPair : TensorManager.GetOutputNameIndexMap())
+	{
+		String += FString::Format(TEXT(" -{0}: {1}\n"), { NameIndexPair.Key, Tensors[NameIndexPair.Value].ToString(20) });
+	}
+	// Add FNeuralOperator(s)
+	String += TEXT("Operators:\n");
+	for (const TSharedPtr<FNeuralOperator>& Operator : Operators)
+	{
+		String += FString::Format(TEXT(" -{0}\n"), { Operator->ToString() });
+	}
+	// Return result
+	return String;
+}
diff --git a/Engine/Plugins/Experimental/NeuralNetworkInference/Source/NeuralNetworkInference/Private/NeuralTensorManager.cpp b/Engine/Plugins/Experimental/NeuralNetworkInference/Source/NeuralNetworkInference/Private/NeuralTensorManager.cpp
new file mode 100644
--- /dev/null
+++ b/Engine/Plugins/Experimental/NeuralNetworkInference/Source/NeuralNetworkInference/Private/NeuralTensorManager.cpp
@@ -0,0 +1,381 @@
+// Copyright Epic Games, Inc. All Rights Reserved.
+
+#include "NeuralTensorManager.h"
+#include "NeuralNetworkInferenceUtils.h"
+#include "NeuralNetworkInferenceVersion.h"
+
+
+
+/* FPrivateNeuralTensorManager
+ *****************************************************************************/
+
+class FPrivateNeuralTensorManager
+{
+public:
+	static TMap<FString, int32> TensorArrayToNameIndexMap(const TArray<FNeuralTensor>& InTensors, const TArray<FNeuralTensor*>& InTensorArray, const FString& InPrefix);
+	static void TensorNameIndexMapToArray(TArray<int32>& OutIndexes, const TMap<FString, int32>& InNameIndexMap);
+};
+
+TMap<FString, int32> FPrivateNeuralTensorManager::TensorArrayToNameIndexMap(const TArray<FNeuralTensor>& InTensors, const TArray<FNeuralTensor*>& InTensorArray, const FString& InPrefix)
+{
+	if (InTensors.Num() >= InTensorArray.Num())
+	{
+		TMap<FString, int32> NameIndexMap;
+		int32 Counter = -1;
+		for (const FNeuralTensor* const Tensor : InTensorArray)
+		{
+			for (int32 TensorIndex = 0; TensorIndex < InTensors.Num(); ++TensorIndex)
+			{
+				if (&InTensors[TensorIndex] == Tensor)
+				{
+					NameIndexMap.Add(InPrefix + FString::FromInt(++Counter), TensorIndex);
+					break;
+				}
+			}
+		}
+		if (NameIndexMap.Num() != InTensorArray.Num())
+		{
+			UE_LOG(LogNeuralNetworkInference, Warning, TEXT("FNeuralTensorManager::TensorArrayToNameIndexMap(): NameIndexMap.Num() == InTensorArray.Num() failed (%d != %d)."),
+				NameIndexMap.Num(), InTensorArray.Num());
+			return TMap<FString, int32>();
+		}
+		return NameIndexMap;
+	}
+	else
+	{
+		UE_LOG(LogNeuralNetworkInference, Warning, TEXT("FNeuralTensorManager::TensorArrayToNameIndexMap(): InTensors.Num() >= InTensorArray.Num() failed (%d vs. %d)."),
+			InTensors.Num(), InTensorArray.Num());
+		return TMap<FString, int32>();
+	}
+}
+
+void FPrivateNeuralTensorManager::TensorNameIndexMapToArray(TArray<int32>& OutIndexes, const TMap<FString, int32>& InNameIndexMap)
+{
+	OutIndexes.Empty();
+	for (const auto& NameIndexPair : InNameIndexMap)
+	{
+		// Add if tensor index only once
+		if (OutIndexes.Find(NameIndexPair.Value) == INDEX_NONE)
+		{
+			OutIndexes.Push(NameIndexPair.Value);
+		}
+	}
+	Algo::Sort(OutIndexes);
+}
+
+
+
+/* FNeuralTensorManager structors
+ *****************************************************************************/
+
+FNeuralTensorManager::FNeuralTensorManager()
+	: bIsLoaded(false)
+{
+}
+
+FNeuralTensorManager::FNeuralTensorManager(TArray<FNeuralTensor>& InTensors, TMap<FString, int32>& InNameIndexMap, TMap<FString, int32>& InInputNameIndexMap,
+	TMap<FString, int32>& InOutputNameIndexMap)
+	: FNeuralTensorManager()
+{
+	Load(InTensors, InNameIndexMap, InInputNameIndexMap, InOutputNameIndexMap);
+}
+
+FNeuralTensorManager::FNeuralTensorManager(TArray<FNeuralTensor>& InTensors, const TArray<FNeuralTensor*>& InInputTensors,
+	const TArray<FNeuralTensor*>& InOutputTensors)
+	: FNeuralTensorManager()
+{
+	Load(InTensors, InInputTensors, InOutputTensors);
+}
+
+FNeuralTensorManager::~FNeuralTensorManager()
+{
+}
+
+
+
+/* FNeuralTensorManager public functions
+ *****************************************************************************/
+
+bool FNeuralTensorManager::Load(TArray<FNeuralTensor>& InTensors, TMap<FString, int32>& InNameIndexMap,
+	TMap<FString, int32>& InInputNameIndexMap, TMap<FString, int32>& InOutputNameIndexMap)
+{
+	bIsLoaded = false;
+	// Sanity checks
+	if (InTensors.Num() < InInputNameIndexMap.Num() || InTensors.Num() < InOutputNameIndexMap.Num())
+	{
+		UE_LOG(LogNeuralNetworkInference, Warning,
+			TEXT("FNeuralTensorManager::Load(): InTensors.Num() >= InputNameIndexMap.Num() (%d vs. %d) || InTensors.Num() >= OutputNameIndexMap.Num() (%d vs. %d) failed."),
+			InTensors.Num(), InInputNameIndexMap.Num(), InTensors.Num(), InOutputNameIndexMap.Num());
+		return bIsLoaded;
+	}
+	else if (InInputNameIndexMap.Num() < 1 || InOutputNameIndexMap.Num() < 1)
+	{
+		UE_LOG(LogNeuralNetworkInference, Warning, TEXT("There should be at least 1 input (currently %d) and 1 output (currently %d) tensors."), InInputNameIndexMap.Num(), InOutputNameIndexMap.Num());
+		return bIsLoaded;
+	}
+	// Fill NameIndexMap, Tensors, InputNameIndexMap, and OutputNameIndexMap
+	Swap(Tensors, InTensors);
+	Swap(NameIndexMap, InNameIndexMap);
+	Swap(InputNameIndexMap, InInputNameIndexMap);
+	Swap(OutputNameIndexMap, InOutputNameIndexMap);
+	// InputIndexes & OutputIndexes
+	FPrivateNeuralTensorManager::TensorNameIndexMapToArray(InputIndexes, InputNameIndexMap);
+	FPrivateNeuralTensorManager::TensorNameIndexMapToArray(OutputIndexes, OutputNameIndexMap);
+	// NonInputIndexes
+	NonInputIndexes.Empty();
+	for (int32 TensorIndex = 0; TensorIndex < Tensors.Num(); ++TensorIndex )
+	{
+		// Add if tensor index is not an input
+		if (InputIndexes.Find(TensorIndex) == INDEX_NONE)
+		{
+			NonInputIndexes.Push(TensorIndex);
+		}
+	}
+	if (InputIndexes.Num() + NonInputIndexes.Num() != Tensors.Num())
+	{
+		UE_LOG(LogNeuralNetworkInference, Warning,
+			TEXT("FNeuralTensorManager::Load(): InputIndexes.Num() + NonInputIndexes.Num() == Tensors.Num() failed (%d + %d != %d), NameIndexMap.Num() = %d)."),
+			InputIndexes.Num(), NonInputIndexes.Num(), Tensors.Num(), NameIndexMap.Num());
+		return bIsLoaded;
+	}
+	// FNeuralTensorManager was configured
+	bIsLoaded = true;
+	// Update Version
+	Version = FNeuralNetworkInferenceVersion::GetVersion();
+	return bIsLoaded;
+}
+
+bool FNeuralTensorManager::Load(TArray<FNeuralTensor>& InTensors, const TArray<FNeuralTensor*>& InInputTensors,
+	const TArray<FNeuralTensor*>& InOutputTensors)
+{
+	TMap<FString, int32> InputNameIndexMapLocal = FPrivateNeuralTensorManager::TensorArrayToNameIndexMap(InTensors, InInputTensors, TEXT("i"));
+	TMap<FString, int32> OutputNameIndexMapLocal = FPrivateNeuralTensorManager::TensorArrayToNameIndexMap(InTensors, InOutputTensors, TEXT("o"));
+	TMap<FString, int32> NameIndexMapLocal = InputNameIndexMapLocal;
+	NameIndexMapLocal.Append(OutputNameIndexMapLocal);
+	// Missing FNeuralTensors (weight & intermediate ones)
+	if (NameIndexMapLocal.Num() < InTensors.Num())
+	{
+		TArray<int32> IndexesLocal;
+		FPrivateNeuralTensorManager::TensorNameIndexMapToArray(IndexesLocal, NameIndexMapLocal);
+		int32 Counter = -1;
+		for (int32 TensorIndex = 0; TensorIndex < InTensors.Num(); ++TensorIndex)
+		{
+			// If not an output nor input --> Add with new name
+			if (IndexesLocal.Find(TensorIndex) == INDEX_NONE)
+			{
+				NameIndexMap.Add(TEXT("t") + FString::FromInt(++Counter), TensorIndex);
+			}
+		}
+	}
+	// Call the original Load() with the new TMaps
+	return Load(InTensors, NameIndexMapLocal, InputNameIndexMapLocal, OutputNameIndexMapLocal);
+}
+
+bool FNeuralTensorManager::IsLoaded() const
+{
+	return bIsLoaded;
+}
+
+TArray<int32> FNeuralTensorManager::GetVersion() const
+{
+	return Version;
+}
+
+const TArray<FNeuralTensor>& FNeuralTensorManager::GetTensors() const
+{
+	SanityCheckIsLoaded();
+	return Tensors;
+}
+
+TArray<FNeuralTensor>& FNeuralTensorManager::GetTensorsMutable()
+{
+	SanityCheckIsLoaded();
+	return Tensors;
+}
+
+const TMap<FString, int32>& FNeuralTensorManager::GetNameIndexMap() const
+{
+	SanityCheckIsLoaded();
+	return NameIndexMap;
+}
+
+const TArray<int32>& FNeuralTensorManager::GetNonInputIndexes() const
+{
+	SanityCheckIsLoaded();
+	return NonInputIndexes;
+}
+
+const FNeuralTensor& FNeuralTensorManager::GetInputTensor() const
+{
+	SanityCheckIsLoaded();
+	if (InputIndexes.Num() != 1)
+	{
+		UE_LOG(LogNeuralNetworkInference, Warning,
+			TEXT("FNeuralTensorManager::GetInputTensor(): There should be 1 input, not %d. Call GetInputIndexes() instead."), InputIndexes.Num());
+	}
+	checkf(Tensors.Num() > 0 && InputIndexes.Num() > 0, TEXT("Tensors and/or InputIndexes was empty."));
+	return Tensors[InputIndexes[0]];
+}
+
+FNeuralTensor& FNeuralTensorManager::GetInputTensorMutable()
+{
+	SanityCheckIsLoaded();
+	if (InputIndexes.Num() != 1)
+	{
+		UE_LOG(LogNeuralNetworkInference, Warning,
+			TEXT("FNeuralTensorManager::GetInputTensorMutable(): There should be 1 input, not %d. Call GetInputIndexes() instead."), InputIndexes.Num());
+	}
+	checkf(Tensors.Num() > 0 && InputIndexes.Num() > 0, TEXT("Tensors and/or InputIndexes was empty."));
+	return Tensors[InputIndexes[0]];
+}
+
+const TArray<int32>& FNeuralTensorManager::GetInputIndexes() const
+{
+	SanityCheckIsLoaded();
+	return InputIndexes;
+}
+
+const TMap<FString, int32>& FNeuralTensorManager::GetInputNameIndexMap() const
+{
+	SanityCheckIsLoaded();
+	return InputNameIndexMap;
+}
+
+void FNeuralTensorManager::SetInputFromTensorCopy(const FNeuralTensor& InTensor)
+{
+	return GetInputTensorMutable().SetFromUnderlyingUInt8ArrayCopy(InTensor.GetUnderlyingUInt8ArrayRef());
+}
+
+void FNeuralTensorManager::SetInputFromTensorMapCopy(const TMap<FString, FNeuralTensor>& InTensorMap)
+{
+	TArray<FNeuralTensor>& TensorsLocalRef = GetTensorsMutable();
+	for (const auto& IndexMapTuple : GetInputNameIndexMap())
+	{
+		TensorsLocalRef[IndexMapTuple.Value].SetFromUnderlyingUInt8ArrayCopy(InTensorMap.FindChecked(IndexMapTuple.Key).GetUnderlyingUInt8ArrayRef());
+	}
+}
+
+TMap<FString, void*> FNeuralTensorManager::CreateInputDataPointersMutable()
+{
+	TMap<FString, void*> InputDataPointers;
+	if (SanityCheckIsLoaded())
+	{
+		for (auto& NameIndexPair : GetInputNameIndexMap())
+		{
+			InputDataPointers.Add(NameIndexPair.Key, Tensors[NameIndexPair.Value].GetData());
+		}
+	}
+	return InputDataPointers;
+}
+
+FRDGBufferUAVRef FNeuralTensorManager::GetInputBufferUAVRef()
+{
+	return GetInputTensorMutable().GetBufferUAVRef();
+}
+
+TMap<FString, FRDGBufferUAVRef> FNeuralTensorManager::CreateInputBufferUAVRefs()
+{
+	TMap<FString, FRDGBufferUAVRef> InputBufferUAVMap;
+	if (SanityCheckIsLoaded())
+	{
+		for (auto& NameIndexPair : GetInputNameIndexMap())
+		{
+			InputBufferUAVMap.Add(NameIndexPair.Key, Tensors[NameIndexPair.Value].GetBufferUAVRef());
+		}
+	}
+	return InputBufferUAVMap;
+}
+
+const FNeuralTensor& FNeuralTensorManager::GetOutputTensor() const
+{
+	SanityCheckIsLoaded();
+	if (OutputIndexes.Num() != 1)
+	{
+		UE_LOG(LogNeuralNetworkInference, Warning,
+			TEXT("FNeuralTensorManager::GetOutputTensor(): There should be 1 input, not %d. Call GetOutputIndexes() instead."), OutputIndexes.Num());
+	}
+	checkf(Tensors.Num() > 0 && OutputIndexes.Num() > 0, TEXT("Tensors and/or OutputIndexes was empty."));
+	return Tensors[OutputIndexes[0]];
+}
+
+FNeuralTensor& FNeuralTensorManager::GetOutputTensorMutable()
+{
+	SanityCheckIsLoaded();
+	if (OutputIndexes.Num() != 1)
+	{
+		UE_LOG(LogNeuralNetworkInference, Warning,
+			TEXT("FNeuralTensorManager::GetOutputTensorMutable(): There should be 1 input, not %d. Call GetOutputIndexes() instead."), OutputIndexes.Num());
+	}
+	checkf(Tensors.Num() > 0 && OutputIndexes.Num() > 0, TEXT("Tensors and/or OutputIndexes was empty."));
+	return Tensors[OutputIndexes[0]];
+}
+
+const TArray<int32>& FNeuralTensorManager::GetOutputIndexes() const
+{
+	SanityCheckIsLoaded();
+	return OutputIndexes;
+}
+
+const TMap<FString, int32>& FNeuralTensorManager::GetOutputNameIndexMap() const
+{
+	SanityCheckIsLoaded();
+	return OutputNameIndexMap;
+}
+
+const FRDGBufferSRVRef FNeuralTensorManager::GetOutputBufferSRVRef() const
+{
+	return GetOutputTensor().GetBufferSRVRef();
+}
+
+TMap<FString, const FRDGBufferSRVRef> FNeuralTensorManager::CreateOutputBufferSRVRefs() const
+{
+	TMap<FString, const FRDGBufferSRVRef> OutputBufferSRVMap;
+	if (SanityCheckIsLoaded())
+	{
+		for (const auto& NameIndexPair : GetOutputNameIndexMap())
+		{
+			OutputBufferSRVMap.Add(NameIndexPair.Key, Tensors[NameIndexPair.Value].GetBufferSRVRef());
+		}
+	}
+	return OutputBufferSRVMap;
+}
+
+TMap<FString, FNeuralTensor> FNeuralTensorManager::CreateInputTensorMap() const
+{
+	TMap<FString, FNeuralTensor> InputArrayMap;
+	if (SanityCheckIsLoaded())
+	{
+		for (const auto& NameIndexPair : GetInputNameIndexMap())
+		{
+			InputArrayMap.Add(NameIndexPair.Key, Tensors[NameIndexPair.Value]);
+		}
+	}
+	return InputArrayMap;
+}
+
+TMap<FString, FNeuralTensor> FNeuralTensorManager::CreateOutputTensorMap() const
+{
+	TMap<FString, FNeuralTensor> OutputArrayMap;
+	if (SanityCheckIsLoaded())
+	{
+		for (const auto& NameIndexPair : GetOutputNameIndexMap())
+		{
+			OutputArrayMap.Add(NameIndexPair.Key, Tensors[NameIndexPair.Value]);
+		}
+	}
+	return OutputArrayMap;
+}
+
+
+
+/* FNeuralTensorManager private functions
+ *****************************************************************************/
+
+bool FNeuralTensorManager::SanityCheckIsLoaded() const
+{
+	if (!bIsLoaded)
+	{
+		UE_LOG(LogNeuralNetworkInference, Warning, TEXT("FNeuralTensorManager::SanityCheckIsLoaded(): Not loaded."));
+	}
+	return bIsLoaded;
+}
diff --git a/Engine/Plugins/Experimental/NeuralNetworkInference/Source/NeuralNetworkInferenceEditor/NeuralNetworkInferenceEditor.Build.cs b/Engine/Plugins/Experimental/NeuralNetworkInference/Source/NeuralNetworkInferenceEditor/NeuralNetworkInferenceEditor.Build.cs
--- a/Engine/Plugins/Experimental/NeuralNetworkInference/Source/NeuralNetworkInferenceEditor/NeuralNetworkInferenceEditor.Build.cs
+++ b/Engine/Plugins/Experimental/NeuralNetworkInference/Source/NeuralNetworkInferenceEditor/NeuralNetworkInferenceEditor.Build.cs
@@ -28,7 +28,6 @@ public class NeuralNetworkInferenceEditor : ModuleRules
 			(
 			new string[] {
 				"NeuralNetworkInference",
-				"NeuralNetworkInferenceLegacy",
 				"UnrealEd"
 			}
 		);
diff --git a/Engine/Plugins/Experimental/NeuralNetworkInference/Source/NeuralNetworkInferenceQA/NeuralNetworkInferenceQA.Build.cs b/Engine/Plugins/Experimental/NeuralNetworkInference/Source/NeuralNetworkInferenceQA/NeuralNetworkInferenceQA.Build.cs
--- a/Engine/Plugins/Experimental/NeuralNetworkInference/Source/NeuralNetworkInferenceQA/NeuralNetworkInferenceQA.Build.cs
+++ b/Engine/Plugins/Experimental/NeuralNetworkInference/Source/NeuralNetworkInferenceQA/NeuralNetworkInferenceQA.Build.cs
@@ -29,8 +29,7 @@ public class NeuralNetworkInferenceQA : ModuleRules
 			(
 			new string[] {
 				"ModelProto",
-				"NeuralNetworkInferenceCore",
-				"NeuralNetworkInferenceLegacy"
+				"NeuralNetworkInferenceCore"
 			}
 		);
 
