NeuralNetworkInference (NNI) plugin:
- Removed NNI QA and NNILegacy from right-click --> Create new ML sth

#rnx

diff --git a/Engine/Plugins/Experimental/NeuralNetworkInference/Source/NeuralNetworkInferenceEditor/Private/NeuralNetworkLegacyFactory.cpp b/Engine/Plugins/Experimental/NeuralNetworkInference/Source/NeuralNetworkInferenceEditor/Private/NeuralNetworkLegacyFactory.cpp
--- a/Engine/Plugins/Experimental/NeuralNetworkInference/Source/NeuralNetworkInferenceEditor/Private/NeuralNetworkLegacyFactory.cpp
+++ b/Engine/Plugins/Experimental/NeuralNetworkInference/Source/NeuralNetworkInferenceEditor/Private/NeuralNetworkLegacyFactory.cpp
@@ -72,7 +72,9 @@ bool UNeuralNetworkLegacyFactory::CanCreateNew() const
 {
 	// If true --> It will always call FactoryCreateNew(), not allowing me to use FactoryCreateFile().
 	// If false --> It will ignore the FactoryCreateFile (thus the txt file) when creating a new UNeuralNetworkLegacy.
-	return CurrentFilename.IsEmpty();
+
+	// Keep false while in production, set to true only when trying to create a new one
+	return false; // CurrentFilename.IsEmpty();
 }
 
 bool UNeuralNetworkLegacyFactory::DoesSupportClass(UClass * Class)
diff --git a/Engine/Plugins/Experimental/NeuralNetworkInference/Source/Private/NeuralNetworkInferenceQAEditor/Private/NeuralNetworkInferenceQAAssetFactory.cpp b/Engine/Plugins/Experimental/NeuralNetworkInference/Source/Private/NeuralNetworkInferenceQAEditor/Private/NeuralNetworkInferenceQAAssetFactory.cpp
--- a/Engine/Plugins/Experimental/NeuralNetworkInference/Source/Private/NeuralNetworkInferenceQAEditor/Private/NeuralNetworkInferenceQAAssetFactory.cpp
+++ b/Engine/Plugins/Experimental/NeuralNetworkInference/Source/Private/NeuralNetworkInferenceQAEditor/Private/NeuralNetworkInferenceQAAssetFactory.cpp
@@ -48,7 +48,9 @@ bool UNeuralNetworkInferenceQAAssetFactory::CanCreateNew() const
 {
 	// If true --> It will always call FactoryCreateNew(), not allowing me to use FactoryCreateFile().
 	// If false --> It will ignore the FactoryCreateFile (thus the txt file) when creating a new UNeuralNetworkInferenceQAAsset.
-	return CurrentFilename.IsEmpty();
+
+	// Keep false while in production, set to true only when trying to create a new one
+	return false; // CurrentFilename.IsEmpty();
 }
 
 bool UNeuralNetworkInferenceQAAssetFactory::DoesSupportClass(UClass * Class)
