NeuralNetworkInference (NNI) plugin:
- Added Doxygen-style doc to main NNI classes

#rnx

diff --git a/Engine/Plugins/Experimental/NeuralNetworkInference/NeuralNetworkInference.uplugin b/Engine/Plugins/Experimental/NeuralNetworkInference/NeuralNetworkInference.uplugin
--- a/Engine/Plugins/Experimental/NeuralNetworkInference/NeuralNetworkInference.uplugin
+++ b/Engine/Plugins/Experimental/NeuralNetworkInference/NeuralNetworkInference.uplugin
@@ -3,7 +3,7 @@
 	"Version": 0.1,
 	"VersionName": "0.1.0-alpha",
 	"FriendlyName": "Neural Network Inference (NNI)",
-	"Description": "Neural Network Inference (NNI) is an experimetal plugin that allows deploying and running deep learning neural network models within Unreal Engine.",
+	"Description": "Cross-platform framework for running deep learning and neural network inference in Unreal Engine, supporting both CPU and GPU acceleration.",
 	"Category": "Experimental",
 	"CreatedBy": "Epic Games, Inc.",
 	"CreatedByURL": "http://epicgames.com",
diff --git a/Engine/Plugins/Experimental/NeuralNetworkInference/Source/NeuralNetworkInference/Internal/NeuralNetworkLegacy.h b/Engine/Plugins/Experimental/NeuralNetworkInference/Source/NeuralNetworkInference/Internal/NeuralNetworkLegacy.h
--- a/Engine/Plugins/Experimental/NeuralNetworkInference/Source/NeuralNetworkInference/Internal/NeuralNetworkLegacy.h
+++ b/Engine/Plugins/Experimental/NeuralNetworkInference/Source/NeuralNetworkInference/Internal/NeuralNetworkLegacy.h
@@ -13,78 +13,8 @@
 class UAssetImportData;
 
 /**
- * UNeuralNetworkLegacy loads and runs a (deep) neural network.
- *
- * 1. Code example of UNeuralNetworkLegacy Configuration:
- *		// -------------------------------------------------- Step 1 - Option a (Editor-only) - UNeuralNetworkLegacy Configuration From ONNX File --------------------------------------------------
- *		#if WITH_EDITOR
- *			// Define the UNeuralNetworkLegacy
- *			UNeuralNetworkLegacy* Network = NewObject<UNeuralNetworkLegacy>((UObject*)GetTransientPackage(), UNeuralNetworkLegacy::StaticClass());
- *			// Try to load the network
- *			if (Network->Load(TEXT("SomeFilePath.onnx"))) // Alternative: Network->Load(...) + Network->IsLoaded()
- *			{
- *				// Set to CPU mode
- *				Network->SetDeviceType(ENeuralDeviceType::CPU);
- *				// Verbose example
- *				UE_LOG(LogNeuralNetworkInference, Display, TEXT("Network loaded:\n%s."), *Network->ToString());
- *			}
- *		#endif
- *
- *		// -------------------------------------------------- Step 1 - Option b - UNeuralNetworkLegacy Configuration From UAsset File --------------------------------------------------
- *		// Define the UNeuralNetworkLegacy
- *		UNeuralNetworkLegacy* Network = LoadObject<UNeuralNetworkLegacy>((UObject*)GetTransientPackage(), TEXT("ExampleNetwork'/Game/Models/ExampleNetwork/ExampleNetwork.ExampleNetwork'"));
- *		// Try to load the network
- *		if (Network->IsLoaded())
- *		{
- *			// Set to CPU mode
- *			Network->SetDeviceType(ENeuralDeviceType::CPU);
- *			// Verbose example
- *			UE_LOG(LogNeuralNetworkInference, Display, TEXT("Network loaded:\n%s."), *Network->ToString());
- *		}
- *
- * 2. Code example of the Forward Pass:
- *		// -------------------------------------------------- Step 2 - Forward Pass - Overview --------------------------------------------------
- *		// Input filling
- *		TArray<float> InArray;
- *		Network->SetInputFromArrayCopy(InArray);
- *		UE_LOG(LogNeuralNetworkInference, Display, TEXT("Input tensor: %s."), *Network->GetInputTensor().ToString());
- *		// Run UNeuralNetworkLegacy
- *		Network->Run();
- *		// Return OutputTensor
- *		const FNeuralTensor& OutputTensor = Network->GetOutputTensor();
- *		UE_LOG(LogNeuralNetworkInference, Display, TEXT("Output tensor: %s."), *OutputTensor.ToString());
- *
- *		// -------------------------------------------------- Step 2 - Forward Pass - Input - Simple vs. Efficient and Single vs. Multiple Filling --------------------------------------------------
- *		// Option a) Simple input filling
- *		TArray<float> InArray;
- *		Network->SetInputFromArrayCopy(InArray);
- *		// Option b) Efficient input filling (it avoids the TArray-FNeuralTensor copy)
- *		float* InputDataPointer = Network->GetInputDataPointerMutable<float>();
- *		for (int64 Index = 0; Index < Network->GetInputTensor().Num(); ++Index)
- *			InputDataPointer[Index] = ...;
- *		// Option c) Simple multiple input filling
- *		TMap<FString, FNeuralTensor> InTensorMap;
- *		Network->SetInputFromTensorMapCopy(InTensorMap);
- *		// Option d) Efficient multiple input filling
- *		TMap<FString, void*> InputDataPointer = Network->CreateInputDataPointersMutable();
- *		for (int64 Index = 0; Index < Network->GetInputTensor().Num(); ++Index)
- *			InputDataPointer[SomeFStringKey][Index] = ...;
- *
- *		// -------------------------------------------------- Step 2 - Forward Pass - Output - Single vs. Multiple Output --------------------------------------------------
- *		// Option a) Return OutputTensor
- *		const FNeuralTensor* OutputTensor = Network->GetOutputTensor();
- *		UE_LOG(LogNeuralNetworkInference, Display, TEXT("Output tensor: %s."), *OutputTensor.ToString());
- *		// Option b) Return OutputTensor
- *		const TArray<FNeuralTensor>& Tensors = GetTensors();
- *		const TMap<FString, int32>& OutputNameIndexMap = Network->GetOutputNameIndexMap();
- *		// Iterate and print on console all output tensors
- *		for (const auto& OutputNameIndexPair)
- *		{
- *			UE_LOG(LogNeuralNetworkInference, Display, TEXT("Output tensor %s: %s."), *OutputNameIndexPair.Key, *Tensors[OutputNameIndexPair.Key].ToString());
- *		}
+ * First version of NNI, meant as a proof of concept.
  */
-//UCLASS(BlueprintType, Deprecated)
-//class NEURALNETWORKINFERENCELEGACY_API UDEPRECATED_UNeuralNetworkLegacy : public UObject
 UCLASS(BlueprintType)
 class NEURALNETWORKINFERENCE_API UNeuralNetworkLegacy : public UObject
 {
diff --git a/Engine/Plugins/Experimental/NeuralNetworkInference/Source/NeuralNetworkInference/Public/NeuralNetwork.h b/Engine/Plugins/Experimental/NeuralNetworkInference/Source/NeuralNetworkInference/Public/NeuralNetwork.h
--- a/Engine/Plugins/Experimental/NeuralNetworkInference/Source/NeuralNetworkInference/Public/NeuralNetwork.h
+++ b/Engine/Plugins/Experimental/NeuralNetworkInference/Source/NeuralNetworkInference/Public/NeuralNetwork.h
@@ -10,7 +10,50 @@
 class UAssetImportData;
 
 /**
- * First version of NNI, meant to work with any ONNX/ORT model, but only in the Editor in Windows
+ * UNeuralNetwork is UE's representation for deep learning and neural network models. It supports the industry standard ONNX model format.
+ * All major frameworks (PyTorch, TensorFlow, MXNet, Caffe2, etc.) provide converters to ONNX.
+ *
+ * See the following examples to learn how to read any ONNX model and run inference (i.e., a forward pass) on it.
+ * 1. Constructing a UNeuralNetwork from an ONNX file (Editor-only):
+ *	#if WITH_EDITOR
+ *		// Create the UNeuralNetwork object
+ *		UNeuralNetwork* Network = NewObject<UNeuralNetwork>((UObject*)GetTransientPackage(), UNeuralNetwork::StaticClass());
+ *		// Try to load the network and set the device (CPU/GPU)
+ *		const FString ONNXModelFilePath = TEXT("SOME_PARENT_FOLDER/SOME_ONNX_FILE_NAME.onnx");
+ *		if (Network->Load(ONNXModelFilePath)) 
+ *		{
+ *			Network->SetDeviceType(ENeuralDeviceType::CPU); // Set to CPU/GPU mode
+ *		}
+ *	#endif
+ *
+ * 2. Loading a UNeuralNetwork from a previously-created UAsset (in Editor or in Game):
+ *		// Create and load the UNeuralNetwork object from a UAsset
+ *		const FString NetworkUAssetFilePath = TEXT("ExampleNetwork'/Game/Models/ExampleNetwork/ExampleNetwork.ExampleNetwork'");
+ *		UNeuralNetwork* Network = LoadObject<UNeuralNetwork>((UObject*)GetTransientPackage(), *NetworkUAssetFilePath);
+ *		// Check that the network was successfully loaded
+ *		check(Network->IsLoaded());
+ *
+ * 3.1. Running inference (i.e., a forward pass):
+ *		// Fill input neural tensor
+ *		TArray<float> InArray;
+ *		Network->SetInputFromArrayCopy(InArray);
+ *		UE_LOG(LogNeuralNetworkInference, Display, TEXT("Input tensor: %s."), *Network->GetInputTensor().ToString());
+ *		// Run UNeuralNetwork
+ *		Network->Run();
+ *		// Read and print OutputTensor
+ *		const FNeuralTensor& OutputTensor = Network->GetOutputTensor();
+ *		UE_LOG(LogNeuralNetworkInference, Display, TEXT("Output tensor: %s."), *OutputTensor.ToString());
+ *
+ * 3.2. Alternative - Filling the input tensor without a TArray-to-FNeuralTensor copy:
+ *		// Obtain input tensor pointer
+ *		float* InputDataPointer = Network->GetInputDataPointerMutable<float>();
+ *		// Fill InputDataPointer
+ *		for (int64 Index = 0; Index < Network->GetInputTensor().Num(); ++Index)
+ *			InputDataPointer[Index] = ...;
+ *
+ * 3.3. Alternative - Networks with multiple input/output tensors:
+* - Multiple inputs: Add InTensorIndex to GetInputTensor(InTensorIndex) or GetInputDataPointerMutable(InTensorIndex) in the examples above, or use GetInputTensors() instead.
+* - Multiple outputs: Add InTensorIndex to GetOutputTensor(InTensorIndex) in the examples above or use GetOutputTensors() instead.
  */
 UCLASS(BlueprintType)
 class NEURALNETWORKINFERENCE_API UNeuralNetwork : public UObject
