NeuralNetworkInference (NNI) plugin:
- Adding (a)sync parameters to NN rather than NNLegacy. NNLegacy will temporary use NN's one.
- Modified MLDeformer accordingly.

#rnx

diff --git a/Engine/Plugins/Experimental/NeuralNetworkInference/Source/NeuralNetworkInference/Internal/NeuralNetworkLegacy.h b/Engine/Plugins/Experimental/NeuralNetworkInference/Source/NeuralNetworkInference/Internal/NeuralNetworkLegacy.h
--- a/Engine/Plugins/Experimental/NeuralNetworkInference/Source/NeuralNetworkInference/Internal/NeuralNetworkLegacy.h
+++ b/Engine/Plugins/Experimental/NeuralNetworkInference/Source/NeuralNetworkInference/Internal/NeuralNetworkLegacy.h
@@ -5,6 +5,7 @@
 #include "CoreMinimal.h"
 #include "ModelProto.h"
 #include "NeuralEnumClasses.h"
+#include "NeuralNetwork.h"
 #include "NeuralOperator.h"
 #include "NeuralTensor.h"
 #include "NeuralTensorManager.h"
@@ -26,15 +27,6 @@ public:
 	 */
 	DECLARE_DELEGATE(FOnAsyncRunCompletedInAnyThread);
 
-	/**
-	 * This enum only applies to GPU mode, where Run() might block the thread until completed, or might run on a background thread.
-	 */
-	enum class EGPUSynchronousMode : uint8
-	{
-		Synchronous, /* Block the thread until the function (e.g., Run) has finished. */
-		Asynchronous /* Do not block the thread (the function will run in the background, e.g., Run). Delegates would let the user know when the function has finished. */
-	};
-
 	UNeuralNetworkLegacy();
 
 	virtual ~UNeuralNetworkLegacy();
@@ -158,7 +150,7 @@ public:
 	 * Its output results can be retrieved with GetOutputTensor() or GetOutputNameIndexMap().
 	 * @param GPUSynchronousMode Whether it should block the thread until the UNeuralNetworkLegacy has fully run.
 	 */
-	void Run(const EGPUSynchronousMode InGPUSynchronousMode = EGPUSynchronousMode::Synchronous, const ENeuralDeviceType InInputDeviceType = ENeuralDeviceType::CPU, const ENeuralDeviceType InOutputDeviceType = ENeuralDeviceType::CPU,
+	void Run(const ENeuralNetworkSynchronousMode InSynchronousMode = ENeuralNetworkSynchronousMode::Synchronous, const ENeuralDeviceType InInputDeviceType = ENeuralDeviceType::CPU, const ENeuralDeviceType InOutputDeviceType = ENeuralDeviceType::CPU,
 		const bool bRunGPUEmptyOnlyForProfiling = false);
 
 	/**
diff --git a/Engine/Plugins/Experimental/NeuralNetworkInference/Source/NeuralNetworkInference/Private/NeuralNetworkLegacy.cpp b/Engine/Plugins/Experimental/NeuralNetworkInference/Source/NeuralNetworkInference/Private/NeuralNetworkLegacy.cpp
--- a/Engine/Plugins/Experimental/NeuralNetworkInference/Source/NeuralNetworkInference/Private/NeuralNetworkLegacy.cpp
+++ b/Engine/Plugins/Experimental/NeuralNetworkInference/Source/NeuralNetworkInference/Private/NeuralNetworkLegacy.cpp
@@ -277,7 +277,7 @@ TMap<FString, FNeuralTensor> UNeuralNetworkLegacy::CreateOutputTensorMap() const
 	return TensorManager.CreateOutputTensorMap();
 }
 
-void UNeuralNetworkLegacy::Run(const EGPUSynchronousMode InGPUSynchronousMode, const ENeuralDeviceType InInputDeviceType, const ENeuralDeviceType InOutputDeviceType, const bool bRunGPUEmptyOnlyForProfiling)
+void UNeuralNetworkLegacy::Run(const ENeuralNetworkSynchronousMode InSynchronousMode, const ENeuralDeviceType InInputDeviceType, const ENeuralDeviceType InOutputDeviceType, const bool bRunGPUEmptyOnlyForProfiling)
 {
 	DECLARE_SCOPE_CYCLE_COUNTER(TEXT("UNeuralNetworkLegacy_Run"), STAT_UNeuralNetworkLegacy_Run, STATGROUP_MachineLearning);
 	// Sanity checks
@@ -316,7 +316,7 @@ void UNeuralNetworkLegacy::Run(const EGPUSynchronousMode InGPUSynchronousMode, c
 
 			// On RHI thread
 			ENQUEUE_RENDER_COMMAND(UNeuralNetworkLegacy_Run_RenderThread)(
-				[this, bRunGPUEmptyOnlyForProfiling, InGPUSynchronousMode, InInputDeviceType, InOutputDeviceType](FRHICommandListImmediate& RHICmdList)
+				[this, bRunGPUEmptyOnlyForProfiling, InSynchronousMode, InInputDeviceType, InOutputDeviceType](FRHICommandListImmediate& RHICmdList)
 				{
 					FMemMark Mark(FMemStack::Get());
 					FRDGBuilder GraphBuilder(RHICmdList, RDG_EVENT_NAME("UNeuralNetworkLegacy::Run()"));
@@ -408,7 +408,7 @@ void UNeuralNetworkLegacy::Run(const EGPUSynchronousMode InGPUSynchronousMode, c
 					}
 
 					// Broadcast delegates (from the render thread)
-					if (InGPUSynchronousMode == EGPUSynchronousMode::Asynchronous)
+					if (InSynchronousMode == ENeuralNetworkSynchronousMode::Asynchronous)
 					{
 						GraphBuilder.AddPass(
 							RDG_EVENT_NAME("Async delegate broadcast"),
@@ -425,7 +425,7 @@ void UNeuralNetworkLegacy::Run(const EGPUSynchronousMode InGPUSynchronousMode, c
 			);
 
 			// Block thread until GPU has finished
-			if (InGPUSynchronousMode == EGPUSynchronousMode::Synchronous)
+			if (InSynchronousMode == ENeuralNetworkSynchronousMode::Synchronous)
 			{
 				FNeuralNetworkInferenceUtils::WaitUntilRHIFinished();
 			}
diff --git a/Engine/Plugins/Experimental/NeuralNetworkInference/Source/NeuralNetworkInference/Public/NeuralNetwork.h b/Engine/Plugins/Experimental/NeuralNetworkInference/Source/NeuralNetworkInference/Public/NeuralNetwork.h
--- a/Engine/Plugins/Experimental/NeuralNetworkInference/Source/NeuralNetworkInference/Public/NeuralNetwork.h
+++ b/Engine/Plugins/Experimental/NeuralNetworkInference/Source/NeuralNetworkInference/Public/NeuralNetwork.h
@@ -9,6 +9,16 @@
 
 class UAssetImportData;
 
+/**
+ * Whether Run() will block the thread until completed, or whether it will run on a background thread, not blocking the calling thread.
+ */
+UENUM()
+enum class ENeuralNetworkSynchronousMode : uint8
+{
+	Synchronous, /* UNeuralNetwork::Run() will block the thread until the network evaluation (i.e., forward pass) has finished. */
+	Asynchronous /* UNeuralNetwork::Run() will initialize a forward pass request on a background thread, not blocking the thread that called it. The user should register to UNeuralNetwork's delegate to know when the forward pass has finished. */
+};
+
 /**
  * UNeuralNetwork is UE's representation for deep learning and neural network models. It supports the industry standard ONNX model format.
  * All major frameworks (PyTorch, TensorFlow, MXNet, Caffe2, etc.) provide converters to ONNX.
@@ -113,8 +123,8 @@ public:
 
 	/**
 	 * Run() executes the forward pass on the current UNeuralNetwork given the current input FDeprecatedNeuralTensor(s), which were previously filled with
-	 * SetInputFromTensorCopy() or GetInputDataPointerMutable(), or their multi-input analogs.
-	 * Its output results can be retrieved with GetOutputTensor() or GetOutputNameIndexMap().
+	 * SetInputFromArrayCopy() or GetInputDataPointerMutable().
+	 * Its output results can be retrieved with GetOutputTensor() or GetOutputTensors().
 	 * @param GPUSynchronousMode Whether it should block the thread until the UNeuralNetwork has fully run.
 	 */
 	void Run();
