Fix missing rename of enum from CL 17183261

diff --git a/Engine/Plugins/Experimental/NeuralNetworkInference/Source/NeuralNetworkInferenceQA/Private/LegacyModelUnitTester.cpp b/Engine/Plugins/Experimental/NeuralNetworkInference/Source/NeuralNetworkInferenceQA/Private/LegacyModelUnitTester.cpp
--- a/Engine/Plugins/Experimental/NeuralNetworkInference/Source/NeuralNetworkInferenceQA/Private/LegacyModelUnitTester.cpp
+++ b/Engine/Plugins/Experimental/NeuralNetworkInference/Source/NeuralNetworkInferenceQA/Private/LegacyModelUnitTester.cpp
@@ -214,7 +214,7 @@ void FLegacyModelUnitTester::ModelAccuracyTest(UNeuralNetworkLegacy* InOutNetwor
 		for (int32 TimerIndex = 0; TimerIndex < InRepetitions; ++TimerIndex) \
 		{ \
 			InOutNetwork->SetInputFromArrayCopy(InputArray); \
-			InOutNetwork->Run(/*GPUSynchronousMode*/UNeuralNetworkLegacy::EGPUSynchronousMode::Synchronous, /*InInputDeviceType*/ENeuralDeviceType::CPU, /*InOutputDeviceType*/ENeuralDeviceType::CPU, /*bRunGPUEmptyOnlyForProfiling*/true); \
+			InOutNetwork->Run(/*GPUSynchronousMode*/ENeuralNetworkSynchronousMode ::Synchronous, /*InInputDeviceType*/ENeuralDeviceType::CPU, /*InOutputDeviceType*/ENeuralDeviceType::CPU, /*bRunGPUEmptyOnlyForProfiling*/true); \
 			OutputArrayGPU = InOutNetwork->GetOutputTensor().GetArrayCopy<float>(); \
 		} \
 	} \
