NeuralNetworkInference (NNI) plugin:
- Preparing NNIQA for automation tests.

#rnx

diff --git a/Engine/Plugins/Experimental/NeuralNetworkInference/Source/NeuralNetworkInferenceQA/Private/LegacyModelUnitTester.h b/Engine/Plugins/Experimental/NeuralNetworkInference/Source/NeuralNetworkInferenceQA/Private/LegacyModelUnitTester.h
--- a/Engine/Plugins/Experimental/NeuralNetworkInference/Source/NeuralNetworkInferenceQA/Private/LegacyModelUnitTester.h
+++ b/Engine/Plugins/Experimental/NeuralNetworkInference/Source/NeuralNetworkInferenceQA/Private/LegacyModelUnitTester.h
@@ -8,20 +8,21 @@
 class FLegacyModelUnitTester
 {
 public:
-	static void GlobalTest(const FString& InModelsDirectory);
+	static void GlobalTest(const FString& InProjectContentDir, const FString& InModelZooRelativeDirectory);
 
 private:
 	/**
 	 * It runs a full model test on the desired model.
 	 */
-	static void ModelLoadAccuracyAndSpeedTests(const FString& InModelsDirectory, const TArray<FString>& InModelNames, const TArray<float>& InInputArrayValues,
-		const TArray<TArray<double>>& InCPUGroundTruths, const TArray<TArray<double>>& InGPUGroundTruths, const TArray<int32>& InCPURepetitions, const TArray<int32>& InGPURepetitions);
+	static void ModelLoadAccuracyAndSpeedTests(const FString& InProjectContentDir, const FString& InModelZooRelativeDirectory, const TArray<FString>& InModelNames,
+		const TArray<float>& InInputArrayValues, const TArray<TArray<double>>& InCPUGroundTruths, const TArray<TArray<double>>& InGPUGroundTruths,
+		const TArray<int32>& InCPURepetitions, const TArray<int32>& InGPURepetitions);
 	/**
 	 * Other auxiliary functions for GlobalTest().
 	 */
-	static FString GetONNXModelFilePath(const FString& InModelsDirectory, const FString& InModelName);
-	static FString GetOtxtModelFilePath(const FString& InModelsDirectory, const FString& InModelName);
-	static FString GetUAssetModelFilePath(const FString& InModelName);
+	static FString GetONNXModelFilePath(const FString& InModelZooDirectory, const FString& InModelName);
+	static FString GetOtxtModelFilePath(const FString& InModelZooDirectory, const FString& InModelName);
+	static FString GetUAssetModelFilePath(const FString& InModelName, const FString& InModelZooRelativeDirectory);
 	static UNeuralNetworkLegacy* NetworkUassetLoadTest(const FString& InUAssetPath);
 	static UNeuralNetworkLegacy* NetworkONNXLoadTest(const FString& InONNXFilePath);
 	static void ModelAccuracyTest(UNeuralNetworkLegacy* InOutNetwork, const TArray<float>& InInputArrayValues, const TArray<double>& InCPUGroundTruths,
diff --git a/Engine/Plugins/Experimental/NeuralNetworkInference/Source/NeuralNetworkInferenceQA/Private/ModelUnitTester.cpp b/Engine/Plugins/Experimental/NeuralNetworkInference/Source/NeuralNetworkInferenceQA/Private/ModelUnitTester.cpp
--- a/Engine/Plugins/Experimental/NeuralNetworkInference/Source/NeuralNetworkInferenceQA/Private/ModelUnitTester.cpp
+++ b/Engine/Plugins/Experimental/NeuralNetworkInference/Source/NeuralNetworkInferenceQA/Private/ModelUnitTester.cpp
@@ -10,7 +10,7 @@
 /* FModelUnitTester static public functions
  *****************************************************************************/
 
-void FModelUnitTester::GlobalTest(const FString& InModelsDirectory)
+void FModelUnitTester::GlobalTest(const FString& InProjectContentDir, const FString& InModelZooRelativeDirectory)
 {
 #ifdef WITH_FULL_NNI_SUPPORT
 	// Model load, accuracy, and speed test
@@ -22,7 +22,7 @@ void FModelUnitTester::GlobalTest(const FString& InModelsDirectory)
 		{138.373184, 126.754100, 127.287398, 130.316194, 127.303495, 124.801134, 126.5462530}, {0.488662, 0.472437, 0.478862, 0.522685, 0.038322, 0.480848, 0.483821}});
 	const TArray<int32> CPURepetitions({ 1000, 1000,  50, 1000 });
 	const TArray<int32> GPURepetitions({ 1000, 1000, 100, 1000 });
-	ModelLoadAccuracyAndSpeedTests(InModelsDirectory, ModelNames, InputArrayValues, CPUGroundTruths, GPUGroundTruths, CPURepetitions, GPURepetitions);
+	ModelLoadAccuracyAndSpeedTests(InProjectContentDir, InModelZooRelativeDirectory, ModelNames, InputArrayValues, CPUGroundTruths, GPUGroundTruths, CPURepetitions, GPURepetitions);
 #else
 	UE_LOG(LogNeuralNetworkInferenceQA, Display, TEXT("--------------- FModelUnitTester test skipped (only if WITH_FULL_NNI_SUPPORT)."));
 #endif //WITH_FULL_NNI_SUPPORT
@@ -33,9 +33,12 @@ void FModelUnitTester::GlobalTest(const FString& InModelsDirectory)
 /* FModelUnitTester static private functions
  *****************************************************************************/
 
-void FModelUnitTester::ModelLoadAccuracyAndSpeedTests(const FString& InModelsDirectory, const TArray<FString>& InModelNames, const TArray<float>& InInputArrayValues,
-	const TArray<TArray<double>>& InCPUGroundTruths, const TArray<TArray<double>>& InGPUGroundTruths, const TArray<int32>& InCPURepetitions, const TArray<int32>& InGPURepetitions)
+void FModelUnitTester::ModelLoadAccuracyAndSpeedTests(const FString& InProjectContentDir, const FString& InModelZooRelativeDirectory, const TArray<FString>& InModelNames,
+	const TArray<float>& InInputArrayValues, const TArray<TArray<double>>& InCPUGroundTruths, const TArray<TArray<double>>& InGPUGroundTruths, const TArray<int32>& InCPURepetitions,
+	const TArray<int32>& InGPURepetitions)
 {
+	const FString ModelZooDirectory = InProjectContentDir / InModelZooRelativeDirectory;
+
 	// Test OTXT/UAsset accuracy
 	for (int32 ModelIndex = 0; ModelIndex < InModelNames.Num(); ++ModelIndex)
 	{
@@ -44,7 +47,7 @@ void FModelUnitTester::ModelLoadAccuracyAndSpeedTests(const FString& InModelsDir
 		const TArray<double>& GPUGroundTruths = InGPUGroundTruths[ModelIndex];
 
 		UE_LOG(LogNeuralNetworkInferenceQA, Display, TEXT("--------------- %s - Network Uasset Load and Run"), *ModelName);
-		const FString UAssetModelFilePath = GetUAssetModelFilePath(ModelName);
+		const FString UAssetModelFilePath = GetUAssetModelFilePath(ModelName, InModelZooRelativeDirectory);
 		UNeuralNetwork* Network = NetworkUassetLoadTest(UAssetModelFilePath);
 		if (!Network)
 		{
@@ -70,13 +73,13 @@ void FModelUnitTester::ModelLoadAccuracyAndSpeedTests(const FString& InModelsDir
 #if WITH_EDITOR
 		for (int32 ONNXOrORTIndex = 0; ONNXOrORTIndex < 2; ++ONNXOrORTIndex)
 		{
-			const FString ModelFilePath = (ONNXOrORTIndex % 2 == 0 ? GetONNXModelFilePath(InModelsDirectory, ModelName) : GetORTModelFilePath(InModelsDirectory, ModelName));
+			const FString ModelFilePath = (ONNXOrORTIndex % 2 == 0 ? GetONNXModelFilePath(ModelZooDirectory, ModelName) : GetORTModelFilePath(ModelZooDirectory, ModelName));
 			const FString ModelType = (ONNXOrORTIndex % 2 == 0 ? TEXT("ONNX") : TEXT("ORT"));
 if (ONNXOrORTIndex == 1)
 {
 	break;
 }
-			UE_LOG(LogNeuralNetworkInferenceQA, Display, TEXT("-------------------- %s - Network %s Load and Run"), *ModelName, *ModelType);
+			UE_LOG(LogNeuralNetworkInferenceQA, Display, TEXT("-------------------- %s - Network %s Load and Run - %s"), *ModelName, *ModelType, *ModelFilePath);
 			ModelAccuracyTest(NetworkONNXOrORTLoadTest(ModelFilePath), InInputArrayValues, CPUGroundTruths, GPUGroundTruths);
 		}
 #else
@@ -89,24 +92,24 @@ if (ONNXOrORTIndex == 1)
 	{
 		const FString& ModelName = InModelNames[ModelIndex];
 		UE_LOG(LogNeuralNetworkInferenceQA, Display, TEXT("--------------- %s - Network UAsset Speed Profiling"), *ModelName);
-		const FString UAssetModelFilePath = GetUAssetModelFilePath(ModelName);
+		const FString UAssetModelFilePath = GetUAssetModelFilePath(ModelName, InModelZooRelativeDirectory);
 		ModelSpeedTest(UAssetModelFilePath, InCPURepetitions[ModelIndex], InGPURepetitions[ModelIndex]);
 	}
 }
 
-FString FModelUnitTester::GetONNXModelFilePath(const FString& InModelsDirectory, const FString& InModelName)
+FString FModelUnitTester::GetONNXModelFilePath(const FString& ModelZooDirectory, const FString& InModelName)
 {
-	return /*FPaths::ConvertRelativePathToFull*/(/*InModelsDirectory*/ TEXT("Models/") / InModelName + TEXT("/") + InModelName + TEXT(".onnx")); // E.g., InModelsDirectory + TEXT("ExampleNetworkReadable/ExampleNetworkReadable.onnx")
+	return /*FPaths::ConvertRelativePathToFull*/(ModelZooDirectory / InModelName + TEXT("/") + InModelName + TEXT(".onnx")); // E.g., ModelZooDirectory / TEXT("ExampleNetworkReadable/ExampleNetworkReadable.onnx")
 }
 
-FString FModelUnitTester::GetORTModelFilePath(const FString& InModelsDirectory, const FString& InModelName)
+FString FModelUnitTester::GetORTModelFilePath(const FString& ModelZooDirectory, const FString& InModelName)
 {
-	return /*FPaths::ConvertRelativePathToFull*/(/*InModelsDirectory*/ TEXT("Models/") / InModelName + TEXT("/") + InModelName + TEXT(".ort")); // E.g., InModelsDirectory + TEXT("ExampleNetworkReadable/ExampleNetworkReadable.ort")
+	return /*FPaths::ConvertRelativePathToFull*/(ModelZooDirectory / InModelName + TEXT("/") + InModelName + TEXT(".ort")); // E.g., ModelZooDirectory / TEXT("ExampleNetworkReadable/ExampleNetworkReadable.ort")
 }
 
-FString FModelUnitTester::GetUAssetModelFilePath(const FString& InModelName)
+FString FModelUnitTester::GetUAssetModelFilePath(const FString& InModelName, const FString& InModelZooRelativeDirectory)
 {
-	return InModelName + TEXT("'/Game/Models/") + InModelName + TEXT("/") + InModelName + TEXT(".") + InModelName + TEXT("'"); // ExampleNetworkReadable'/Game/Models/ExampleNetworkReadable/ExampleNetworkReadable.ExampleNetworkReadable'
+	return InModelName + TEXT("'/Game/") + InModelZooRelativeDirectory / InModelName + TEXT("/") + InModelName + TEXT(".") + InModelName + TEXT("'"); // E.g., '/Game/[MODEL_ZOO_DIR]/ExampleNetworkReadable/ExampleNetworkReadable.ExampleNetworkReadable'
 }
 
 UNeuralNetwork* FModelUnitTester::NetworkUassetLoadTest(const FString& InUassetPath)
@@ -135,7 +138,7 @@ UNeuralNetwork* FModelUnitTester::NetworkONNXOrORTLoadTest(const FString& InMode
 		ensureMsgf(false, TEXT("UNeuralNetwork is a nullptr."));
 		return nullptr;
 	}
-	if (!Network->Load(FPaths::ProjectContentDir() / InModelFilePath))
+	if (!Network->Load(InModelFilePath))
 	{
 		ensureMsgf(false, TEXT("UNeuralNetwork could not be loaded from ONNX file disk location: %s"), *InModelFilePath);
 		return nullptr;
diff --git a/Engine/Plugins/Experimental/NeuralNetworkInference/Source/NeuralNetworkInferenceQA/Private/ModelUnitTester.h b/Engine/Plugins/Experimental/NeuralNetworkInference/Source/NeuralNetworkInferenceQA/Private/ModelUnitTester.h
--- a/Engine/Plugins/Experimental/NeuralNetworkInference/Source/NeuralNetworkInferenceQA/Private/ModelUnitTester.h
+++ b/Engine/Plugins/Experimental/NeuralNetworkInference/Source/NeuralNetworkInferenceQA/Private/ModelUnitTester.h
@@ -8,20 +8,21 @@
 class FModelUnitTester
 {
 public:
-	static void GlobalTest(const FString& InModelsDirectory);
+	static void GlobalTest(const FString& InProjectContentDir, const FString& InModelZooRelativeDirectory);
 
 private:
 	/**
 	 * It runs a full model test on the desired model.
 	 */
-	static void ModelLoadAccuracyAndSpeedTests(const FString& InModelsDirectory, const TArray<FString>& InModelNames, const TArray<float>& InInputArrayValues,
-		const TArray<TArray<double>>& InCPUGroundTruths, const TArray<TArray<double>>& InGPUGroundTruths, const TArray<int32>& InCPURepetitions, const TArray<int32>& InGPURepetitions);
+	static void ModelLoadAccuracyAndSpeedTests(const FString& InProjectContentDir, const FString& InModelZooRelativeDirectory, const TArray<FString>& InModelNames,
+		const TArray<float>& InInputArrayValues, const TArray<TArray<double>>& InCPUGroundTruths, const TArray<TArray<double>>& InGPUGroundTruths,
+		const TArray<int32>& InCPURepetitions, const TArray<int32>& InGPURepetitions);
 	/**
 	 * Other auxiliary functions for GlobalTest().
 	 */
-	static FString GetONNXModelFilePath(const FString& InModelsDirectory, const FString& InModelName);
-	static FString GetORTModelFilePath(const FString& InModelsDirectory, const FString& InModelName);
-	static FString GetUAssetModelFilePath(const FString& InModelName);
+	static FString GetONNXModelFilePath(const FString& InModelZooDirectory, const FString& InModelName);
+	static FString GetORTModelFilePath(const FString& InModelZooDirectory, const FString& InModelName);
+	static FString GetUAssetModelFilePath(const FString& InModelName, const FString& InModelZooRelativeDirectory);
 	static UNeuralNetwork* NetworkUassetLoadTest(const FString& InUassetPath);
 #if WITH_EDITOR
 	static UNeuralNetwork* NetworkONNXOrORTLoadTest(const FString& InModelFilePath);
diff --git a/Engine/Plugins/Experimental/NeuralNetworkInference/Source/NeuralNetworkInferenceQA/Private/NeuralNetworkInferenceQA.cpp b/Engine/Plugins/Experimental/NeuralNetworkInference/Source/NeuralNetworkInferenceQA/Private/NeuralNetworkInferenceQA.cpp
--- a/Engine/Plugins/Experimental/NeuralNetworkInference/Source/NeuralNetworkInferenceQA/Private/NeuralNetworkInferenceQA.cpp
+++ b/Engine/Plugins/Experimental/NeuralNetworkInference/Source/NeuralNetworkInferenceQA/Private/NeuralNetworkInferenceQA.cpp
@@ -12,7 +12,9 @@
 
 void UNeuralNetworkInferenceQA::UnitTesting()
 {
-	const FString ModelsDirectory = FPaths::ProjectContentDir() / TEXT("Models/");
-	const FString UnitTestingDirectory = FPaths::ProjectContentDir() / TEXT("UnitTesting/");
-	FUnitTester::GlobalTest(UnitTestingDirectory, ModelsDirectory);
+	const FString ProjectContentDir = FPaths::ProjectContentDir();
+	const FString MachineLearningTestsRelativeDirectory = TEXT("Tests/MachineLearning/");
+	const FString ModelZooRelativeDirectory = MachineLearningTestsRelativeDirectory / TEXT("Models/"); // Eg "Tests/MachineLearning/Models/"
+	const FString UnitTestRelativeDirectory = MachineLearningTestsRelativeDirectory / TEXT("UnitTests/"); // Eg "Tests/MachineLearning/UnitTests/"
+	FUnitTester::GlobalTest(ProjectContentDir, ModelZooRelativeDirectory, UnitTestRelativeDirectory);
 }
diff --git a/Engine/Plugins/Experimental/NeuralNetworkInference/Source/NeuralNetworkInferenceQA/Private/NeuralNetworkInferenceQAAsset.cpp b/Engine/Plugins/Experimental/NeuralNetworkInference/Source/NeuralNetworkInferenceQA/Private/NeuralNetworkInferenceQAAsset.cpp
--- a/Engine/Plugins/Experimental/NeuralNetworkInference/Source/NeuralNetworkInferenceQA/Private/NeuralNetworkInferenceQAAsset.cpp
+++ b/Engine/Plugins/Experimental/NeuralNetworkInference/Source/NeuralNetworkInferenceQA/Private/NeuralNetworkInferenceQAAsset.cpp
@@ -219,7 +219,7 @@ const FString UNeuralNetworkInferenceQAAsset_FileName = TEXT("NeuralNetworkInfer
 UNeuralNetworkInferenceQAAsset* UNeuralNetworkInferenceQAAsset::Load(const FString& InNeuralNetworkInferenceQAAssetParentDirectoryName, const FString& InNeuralNetworkInferenceQAAssetName)
 {
 	// Load NeuralNetworkInferenceQAAsset from disk
-	const FString FilePath = UNeuralNetworkInferenceQAAsset_FileName + TEXT("'/Game/") + InNeuralNetworkInferenceQAAssetParentDirectoryName + TEXT("/") + UNeuralNetworkInferenceQAAsset_FileName + TEXT(".") + UNeuralNetworkInferenceQAAsset_FileName + TEXT("'"); // "NeuralNetworkInferenceQAAsset'/Game/UnitTesting/NeuralNetworkInferenceQAAsset.NeuralNetworkInferenceQAAsset'"
+	const FString FilePath = UNeuralNetworkInferenceQAAsset_FileName + TEXT("'/Game/") / InNeuralNetworkInferenceQAAssetParentDirectoryName / UNeuralNetworkInferenceQAAsset_FileName + TEXT(".") + UNeuralNetworkInferenceQAAsset_FileName + TEXT("'"); // "NeuralNetworkInferenceQAAsset'/Game/[UNIT_TEST_DIR]/NeuralNetworkInferenceQAAsset.NeuralNetworkInferenceQAAsset'"
 	UNeuralNetworkInferenceQAAsset* NeuralNetworkInferenceQAAsset = LoadObject<UNeuralNetworkInferenceQAAsset>((UObject*)GetTransientPackage(), *FilePath);
 	// Warning if it does not exist yet
 	if (!NeuralNetworkInferenceQAAsset)
diff --git a/Engine/Plugins/Experimental/NeuralNetworkInference/Source/NeuralNetworkInferenceQA/Private/OperatorUnitTester.cpp b/Engine/Plugins/Experimental/NeuralNetworkInference/Source/NeuralNetworkInferenceQA/Private/OperatorUnitTester.cpp
--- a/Engine/Plugins/Experimental/NeuralNetworkInference/Source/NeuralNetworkInferenceQA/Private/OperatorUnitTester.cpp
+++ b/Engine/Plugins/Experimental/NeuralNetworkInference/Source/NeuralNetworkInferenceQA/Private/OperatorUnitTester.cpp
@@ -14,12 +14,13 @@
 /* FOperatorUnitTester static public functions
  *****************************************************************************/
 
-void FOperatorUnitTester::GlobalTest(const FString& InGroundTruthDirectory)
+void FOperatorUnitTester::GlobalTest(const FString& InProjectContentDir, const FString& InUnitTestRelativeDirectory)
 {
+	const FString UnitTestDirectory = InProjectContentDir / InUnitTestRelativeDirectory;
 	const float ZeroThreshold = 0.f;
 	const float FloatThreshold = 5e-7;
 	const float RelaxedFloatThreshold = 5e-6;
-	const float VeryRelaxedFloatThreshold = 2e-5;
+	const float VeryRelaxedFloatThreshold = 4e-5;
 	const float VeryVeryRelaxedFloatThreshold = 2e-4;
 	const float SuperRelaxedFloatThreshold = 1.5e-3;
 	// Input arrays
@@ -139,7 +140,7 @@ void FOperatorUnitTester::GlobalTest(const FString& InGroundTruthDirectory)
 		TEXT("FRoundOperator"), TEXT("FSigmoidOperator"), TEXT("FSignOperator"), TEXT("FSinOperator"), TEXT("FSinhOperator"), TEXT("FSqrtOperator"), TEXT("FSqueezeOperator"), TEXT("FSubOperator"), 
 		TEXT("FTanOperator"), TEXT("FTanhOperator") });
 	// Creating NeuralNetworkInferenceQAAsset
-	UNeuralNetworkInferenceQAAsset* NeuralNetworkInferenceQAAsset = UNeuralNetworkInferenceQAAsset::Load(TEXT("UnitTesting"), TEXT("NeuralNetworkInferenceQAAsset"));
+	UNeuralNetworkInferenceQAAsset* NeuralNetworkInferenceQAAsset = UNeuralNetworkInferenceQAAsset::Load(InUnitTestRelativeDirectory, TEXT("NeuralNetworkInferenceQAAsset"));
 	if (!NeuralNetworkInferenceQAAsset)
 	{
 		ensureMsgf(false, TEXT("NeuralNetworkInferenceQAAsset could not be loaded!"));
@@ -150,51 +151,51 @@ void FOperatorUnitTester::GlobalTest(const FString& InGroundTruthDirectory)
 	// Per operator testing (checking false if errors occur)
 	const int32 OperatorTotal = OperatorNames.Num(); // 164;
 	int32 OperatorCounter = 0;
-	FOperatorTester::TestElementWiseOperator<FAbsOperator>(NeuralNetworkInferenceQAAsset, TEXT("FAbsOperator"), ElementWiseTensors, ZeroThreshold, InGroundTruthDirectory, ++OperatorCounter, OperatorTotal);
-	FOperatorTester::TestElementWiseOperator<FAcosOperator>(NeuralNetworkInferenceQAAsset, TEXT("FAcosOperator"), ElementWiseTensorsTiny, VeryVeryRelaxedFloatThreshold, InGroundTruthDirectory, ++OperatorCounter, OperatorTotal);
-	FOperatorTester::TestMultiDirectionalBroadcastingOperator<FAddOperator>(NeuralNetworkInferenceQAAsset, TEXT("FAddOperator"), MultidirectionalTensors, ZeroThreshold, InGroundTruthDirectory, ++OperatorCounter, OperatorTotal);
-	FOperatorTester::TestElementWiseOperator<FAsinOperator>(NeuralNetworkInferenceQAAsset, TEXT("FAsinOperator"), ElementWiseTensorsTiny, VeryVeryRelaxedFloatThreshold, InGroundTruthDirectory, ++OperatorCounter, OperatorTotal);
-	FOperatorTester::TestElementWiseOperator<FAtanOperator>(NeuralNetworkInferenceQAAsset, TEXT("FAtanOperator"), ElementWiseTensors, VeryVeryRelaxedFloatThreshold, InGroundTruthDirectory, ++OperatorCounter, OperatorTotal);
-	FOperatorTester::TestBatchNormalizationOperator(NeuralNetworkInferenceQAAsset, InputTensorData, InputTensorDataSmall, InputTensorDataPos, FloatThreshold, InGroundTruthDirectory, ++OperatorCounter, OperatorTotal);
-	FOperatorTester::TestElementWiseOperator<FCeilOperator>(NeuralNetworkInferenceQAAsset, TEXT("FCeilOperator"), ElementWiseTensors, ZeroThreshold, InGroundTruthDirectory, ++OperatorCounter, OperatorTotal);
-	FOperatorTester::TestConvOperator(NeuralNetworkInferenceQAAsset, InputTensorData, InputTensorDataSmall, FloatThreshold, InGroundTruthDirectory, ++OperatorCounter, OperatorTotal);
-	FOperatorTester::TestConvTransposeOperator(NeuralNetworkInferenceQAAsset, InputTensorData, InputTensorDataSmall, FloatThreshold, InGroundTruthDirectory, ++OperatorCounter, OperatorTotal);
-	FOperatorTester::TestElementWiseOperator<FCosOperator>(NeuralNetworkInferenceQAAsset, TEXT("FCosOperator"), ElementWiseTensors, VeryRelaxedFloatThreshold, InGroundTruthDirectory, ++OperatorCounter, OperatorTotal);
-	FOperatorTester::TestElementWiseOperator<FCoshOperator>(NeuralNetworkInferenceQAAsset, TEXT("FCoshOperator"), ElementWiseTensorsSmall, FloatThreshold, InGroundTruthDirectory, ++OperatorCounter, OperatorTotal);
-	FOperatorTester::TestMultiDirectionalBroadcastingOperator<FDivOperator>(NeuralNetworkInferenceQAAsset, TEXT("FDivOperator"), MultidirectionalTensors, FloatThreshold, InGroundTruthDirectory, ++OperatorCounter, OperatorTotal);
-	FOperatorTester::TestElementWiseOperator<FExpOperator>(NeuralNetworkInferenceQAAsset, TEXT("FExpOperator"), ElementWiseTensorsSmall, FloatThreshold, InGroundTruthDirectory, ++OperatorCounter, OperatorTotal);
-	FOperatorTester::TestElementWiseOperator<FFloorOperator>(NeuralNetworkInferenceQAAsset, TEXT("FFloorOperator"), ElementWiseTensors, ZeroThreshold, InGroundTruthDirectory, ++OperatorCounter, OperatorTotal);
+	FOperatorTester::TestElementWiseOperator<FAbsOperator>(NeuralNetworkInferenceQAAsset, TEXT("FAbsOperator"), ElementWiseTensors, ZeroThreshold, UnitTestDirectory, ++OperatorCounter, OperatorTotal);
+	FOperatorTester::TestElementWiseOperator<FAcosOperator>(NeuralNetworkInferenceQAAsset, TEXT("FAcosOperator"), ElementWiseTensorsTiny, VeryVeryRelaxedFloatThreshold, UnitTestDirectory, ++OperatorCounter, OperatorTotal);
+	FOperatorTester::TestMultiDirectionalBroadcastingOperator<FAddOperator>(NeuralNetworkInferenceQAAsset, TEXT("FAddOperator"), MultidirectionalTensors, ZeroThreshold, UnitTestDirectory, ++OperatorCounter, OperatorTotal);
+	FOperatorTester::TestElementWiseOperator<FAsinOperator>(NeuralNetworkInferenceQAAsset, TEXT("FAsinOperator"), ElementWiseTensorsTiny, VeryVeryRelaxedFloatThreshold, UnitTestDirectory, ++OperatorCounter, OperatorTotal);
+	FOperatorTester::TestElementWiseOperator<FAtanOperator>(NeuralNetworkInferenceQAAsset, TEXT("FAtanOperator"), ElementWiseTensors, VeryVeryRelaxedFloatThreshold, UnitTestDirectory, ++OperatorCounter, OperatorTotal);
+	FOperatorTester::TestBatchNormalizationOperator(NeuralNetworkInferenceQAAsset, InputTensorData, InputTensorDataSmall, InputTensorDataPos, FloatThreshold, UnitTestDirectory, ++OperatorCounter, OperatorTotal);
+	FOperatorTester::TestElementWiseOperator<FCeilOperator>(NeuralNetworkInferenceQAAsset, TEXT("FCeilOperator"), ElementWiseTensors, ZeroThreshold, UnitTestDirectory, ++OperatorCounter, OperatorTotal);
+	FOperatorTester::TestConvOperator(NeuralNetworkInferenceQAAsset, InputTensorData, InputTensorDataSmall, FloatThreshold, UnitTestDirectory, ++OperatorCounter, OperatorTotal);
+	FOperatorTester::TestConvTransposeOperator(NeuralNetworkInferenceQAAsset, InputTensorData, InputTensorDataSmall, FloatThreshold, UnitTestDirectory, ++OperatorCounter, OperatorTotal);
+	FOperatorTester::TestElementWiseOperator<FCosOperator>(NeuralNetworkInferenceQAAsset, TEXT("FCosOperator"), ElementWiseTensors, VeryRelaxedFloatThreshold, UnitTestDirectory, ++OperatorCounter, OperatorTotal);
+	FOperatorTester::TestElementWiseOperator<FCoshOperator>(NeuralNetworkInferenceQAAsset, TEXT("FCoshOperator"), ElementWiseTensorsSmall, FloatThreshold, UnitTestDirectory, ++OperatorCounter, OperatorTotal);
+	FOperatorTester::TestMultiDirectionalBroadcastingOperator<FDivOperator>(NeuralNetworkInferenceQAAsset, TEXT("FDivOperator"), MultidirectionalTensors, FloatThreshold, UnitTestDirectory, ++OperatorCounter, OperatorTotal);
+	FOperatorTester::TestElementWiseOperator<FExpOperator>(NeuralNetworkInferenceQAAsset, TEXT("FExpOperator"), ElementWiseTensorsSmall, FloatThreshold, UnitTestDirectory, ++OperatorCounter, OperatorTotal);
+	FOperatorTester::TestElementWiseOperator<FFloorOperator>(NeuralNetworkInferenceQAAsset, TEXT("FFloorOperator"), ElementWiseTensors, ZeroThreshold, UnitTestDirectory, ++OperatorCounter, OperatorTotal);
 	//UE_LOG(LogNeuralNetworkInferenceQA, Display, TEXT("--------------- %d/%d FGather"), ++OperatorCounter, OperatorTotal);
 	//{
 	//	UE_LOG(LogNeuralNetworkInferenceQA, Display, TEXT("-------------------- Not implemented (yet)"));
 	//}
-	FOperatorTester::TestGemmOperator(NeuralNetworkInferenceQAAsset, InputTensorData, FloatThreshold, InGroundTruthDirectory, ++OperatorCounter, OperatorTotal);
-	FOperatorTester::TestElementWise1AttributeOperator<FLeakyReluOperator>(NeuralNetworkInferenceQAAsset, TEXT("FLeakyReluOperator"), ElementWiseTensors, ZeroThreshold, InGroundTruthDirectory, ++OperatorCounter, OperatorTotal);
-	FOperatorTester::TestElementWiseOperator<FLogOperator>(NeuralNetworkInferenceQAAsset, TEXT("FLogOperator"), ElementWiseTensors, FloatThreshold, InGroundTruthDirectory, ++OperatorCounter, OperatorTotal);
-	FOperatorTester::TestMultiDirectionalBroadcastingOperator<FMulOperator>(NeuralNetworkInferenceQAAsset, TEXT("FMulOperator"), MultidirectionalTensors, ZeroThreshold, InGroundTruthDirectory, ++OperatorCounter, OperatorTotal);
-	FOperatorTester::TestElementWiseOperator<FNegOperator>(NeuralNetworkInferenceQAAsset, TEXT("FNegOperator"), ElementWiseTensors, ZeroThreshold, InGroundTruthDirectory, ++OperatorCounter, OperatorTotal);
-	FOperatorTester::TestMultiDirectionalBroadcastingOperator<FPowOperator>(NeuralNetworkInferenceQAAsset, TEXT("FPowOperator"), MultidirectionalTensors, RelaxedFloatThreshold, InGroundTruthDirectory, ++OperatorCounter, OperatorTotal);
-	FOperatorTester::TestElementWiseOperator<FReciprocalOperator>(NeuralNetworkInferenceQAAsset, TEXT("FReciprocalOperator"), ElementWiseTensors, FloatThreshold, InGroundTruthDirectory, ++OperatorCounter, OperatorTotal);
+	FOperatorTester::TestGemmOperator(NeuralNetworkInferenceQAAsset, InputTensorData, FloatThreshold, UnitTestDirectory, ++OperatorCounter, OperatorTotal);
+	FOperatorTester::TestElementWise1AttributeOperator<FLeakyReluOperator>(NeuralNetworkInferenceQAAsset, TEXT("FLeakyReluOperator"), ElementWiseTensors, ZeroThreshold, UnitTestDirectory, ++OperatorCounter, OperatorTotal);
+	FOperatorTester::TestElementWiseOperator<FLogOperator>(NeuralNetworkInferenceQAAsset, TEXT("FLogOperator"), ElementWiseTensors, FloatThreshold, UnitTestDirectory, ++OperatorCounter, OperatorTotal);
+	FOperatorTester::TestMultiDirectionalBroadcastingOperator<FMulOperator>(NeuralNetworkInferenceQAAsset, TEXT("FMulOperator"), MultidirectionalTensors, ZeroThreshold, UnitTestDirectory, ++OperatorCounter, OperatorTotal);
+	FOperatorTester::TestElementWiseOperator<FNegOperator>(NeuralNetworkInferenceQAAsset, TEXT("FNegOperator"), ElementWiseTensors, ZeroThreshold, UnitTestDirectory, ++OperatorCounter, OperatorTotal);
+	FOperatorTester::TestMultiDirectionalBroadcastingOperator<FPowOperator>(NeuralNetworkInferenceQAAsset, TEXT("FPowOperator"), MultidirectionalTensors, RelaxedFloatThreshold, UnitTestDirectory, ++OperatorCounter, OperatorTotal);
+	FOperatorTester::TestElementWiseOperator<FReciprocalOperator>(NeuralNetworkInferenceQAAsset, TEXT("FReciprocalOperator"), ElementWiseTensors, FloatThreshold, UnitTestDirectory, ++OperatorCounter, OperatorTotal);
 	//UE_LOG(LogNeuralNetworkInferenceQA, Display, TEXT("--------------- %d/%d FReduceSum"), ++OperatorCounter, OperatorTotal);
 	//{
 	//	UE_LOG(LogNeuralNetworkInferenceQA, Display, TEXT("-------------------- Not implemented (yet)"));
 	//}
-	FOperatorTester::TestElementWiseOperator<FReluOperator>(NeuralNetworkInferenceQAAsset, TEXT("FReluOperator"), ElementWiseTensors, ZeroThreshold, InGroundTruthDirectory, ++OperatorCounter, OperatorTotal);
-	FOperatorTester::TestReshapeOperator(NeuralNetworkInferenceQAAsset, InputTensorData, InputTensorDataPos, ZeroThreshold, InGroundTruthDirectory, ++OperatorCounter, OperatorTotal);
-	FOperatorTester::TestElementWiseOperator<FRoundOperator>(NeuralNetworkInferenceQAAsset, TEXT("FRoundOperator"), ElementWiseTensors, ZeroThreshold, InGroundTruthDirectory, ++OperatorCounter, OperatorTotal);
+	FOperatorTester::TestElementWiseOperator<FReluOperator>(NeuralNetworkInferenceQAAsset, TEXT("FReluOperator"), ElementWiseTensors, ZeroThreshold, UnitTestDirectory, ++OperatorCounter, OperatorTotal);
+	FOperatorTester::TestReshapeOperator(NeuralNetworkInferenceQAAsset, InputTensorData, InputTensorDataPos, ZeroThreshold, UnitTestDirectory, ++OperatorCounter, OperatorTotal);
+	FOperatorTester::TestElementWiseOperator<FRoundOperator>(NeuralNetworkInferenceQAAsset, TEXT("FRoundOperator"), ElementWiseTensors, ZeroThreshold, UnitTestDirectory, ++OperatorCounter, OperatorTotal);
 	//UE_LOG(LogNeuralNetworkInferenceQA, Display, TEXT("--------------- %d/%d FShape"), ++OperatorCounter, OperatorTotal);
 	//{
 	//	UE_LOG(LogNeuralNetworkInferenceQA, Display, TEXT("-------------------- Not implemented (yet)"));
 	//}
-	FOperatorTester::TestElementWiseOperator<FSigmoidOperator>(NeuralNetworkInferenceQAAsset, TEXT("FSigmoidOperator"), ElementWiseTensors, FloatThreshold, InGroundTruthDirectory, ++OperatorCounter, OperatorTotal);
-	FOperatorTester::TestElementWiseOperator<FSignOperator>(NeuralNetworkInferenceQAAsset, TEXT("FSignOperator"), ElementWiseTensors, ZeroThreshold, InGroundTruthDirectory, ++OperatorCounter, OperatorTotal);
-	FOperatorTester::TestElementWiseOperator<FSinOperator>(NeuralNetworkInferenceQAAsset, TEXT("FSinOperator"), ElementWiseTensors, VeryRelaxedFloatThreshold, InGroundTruthDirectory, ++OperatorCounter, OperatorTotal);
-	FOperatorTester::TestElementWiseOperator<FSinhOperator>(NeuralNetworkInferenceQAAsset, TEXT("FSinhOperator"), ElementWiseTensorsSmall, FloatThreshold, InGroundTruthDirectory, ++OperatorCounter, OperatorTotal);
-	FOperatorTester::TestElementWiseOperator<FSqrtOperator>(NeuralNetworkInferenceQAAsset, TEXT("FSqrtOperator"), ElementWiseTensors, FloatThreshold, InGroundTruthDirectory, ++OperatorCounter, OperatorTotal);
-	FOperatorTester::TestSqueezeOperator(NeuralNetworkInferenceQAAsset, InputTensorData, InputTensorDataPos, ZeroThreshold, InGroundTruthDirectory, ++OperatorCounter, OperatorTotal);
-	FOperatorTester::TestMultiDirectionalBroadcastingOperator<FSubOperator>(NeuralNetworkInferenceQAAsset, TEXT("FSubOperator"), MultidirectionalTensors, ZeroThreshold, InGroundTruthDirectory, ++OperatorCounter, OperatorTotal);
-	FOperatorTester::TestElementWiseOperator<FTanOperator>(NeuralNetworkInferenceQAAsset, TEXT("FTanOperator"), ElementWiseTensors, SuperRelaxedFloatThreshold, InGroundTruthDirectory, ++OperatorCounter, OperatorTotal);
-	FOperatorTester::TestElementWiseOperator<FTanhOperator>(NeuralNetworkInferenceQAAsset, TEXT("FTanhOperator"), ElementWiseTensorsSmall, FloatThreshold, InGroundTruthDirectory, ++OperatorCounter, OperatorTotal);
+	FOperatorTester::TestElementWiseOperator<FSigmoidOperator>(NeuralNetworkInferenceQAAsset, TEXT("FSigmoidOperator"), ElementWiseTensors, FloatThreshold, UnitTestDirectory, ++OperatorCounter, OperatorTotal);
+	FOperatorTester::TestElementWiseOperator<FSignOperator>(NeuralNetworkInferenceQAAsset, TEXT("FSignOperator"), ElementWiseTensors, ZeroThreshold, UnitTestDirectory, ++OperatorCounter, OperatorTotal);
+	FOperatorTester::TestElementWiseOperator<FSinOperator>(NeuralNetworkInferenceQAAsset, TEXT("FSinOperator"), ElementWiseTensors, VeryRelaxedFloatThreshold, UnitTestDirectory, ++OperatorCounter, OperatorTotal);
+	FOperatorTester::TestElementWiseOperator<FSinhOperator>(NeuralNetworkInferenceQAAsset, TEXT("FSinhOperator"), ElementWiseTensorsSmall, FloatThreshold, UnitTestDirectory, ++OperatorCounter, OperatorTotal);
+	FOperatorTester::TestElementWiseOperator<FSqrtOperator>(NeuralNetworkInferenceQAAsset, TEXT("FSqrtOperator"), ElementWiseTensors, FloatThreshold, UnitTestDirectory, ++OperatorCounter, OperatorTotal);
+	FOperatorTester::TestSqueezeOperator(NeuralNetworkInferenceQAAsset, InputTensorData, InputTensorDataPos, ZeroThreshold, UnitTestDirectory, ++OperatorCounter, OperatorTotal);
+	FOperatorTester::TestMultiDirectionalBroadcastingOperator<FSubOperator>(NeuralNetworkInferenceQAAsset, TEXT("FSubOperator"), MultidirectionalTensors, ZeroThreshold, UnitTestDirectory, ++OperatorCounter, OperatorTotal);
+	FOperatorTester::TestElementWiseOperator<FTanOperator>(NeuralNetworkInferenceQAAsset, TEXT("FTanOperator"), ElementWiseTensors, SuperRelaxedFloatThreshold, UnitTestDirectory, ++OperatorCounter, OperatorTotal);
+	FOperatorTester::TestElementWiseOperator<FTanhOperator>(NeuralNetworkInferenceQAAsset, TEXT("FTanhOperator"), ElementWiseTensorsSmall, FloatThreshold, UnitTestDirectory, ++OperatorCounter, OperatorTotal);
 	//UE_LOG(LogNeuralNetworkInferenceQA, Display, TEXT("--------------- %d/%d FUnsqueeze"), ++OperatorCounter, OperatorTotal);
 	//{
 	//	UE_LOG(LogNeuralNetworkInferenceQA, Display, TEXT("-------------------- Not implemented (yet)"));
@@ -203,7 +204,7 @@ void FOperatorUnitTester::GlobalTest(const FString& InGroundTruthDirectory)
 	ensureMsgf(OperatorTotal == OperatorCounter, TEXT("New or missing operators found, OperatorTotal == OperatorCounter failed (%d != %d)"), OperatorTotal, OperatorCounter);
 
 	// Test against GT saved on disk - Only saving a unique UAsset file with all results from all operators
-	if (!NeuralNetworkInferenceQAAsset->CompareNewVsPreviousTests(InGroundTruthDirectory))
+	if (!NeuralNetworkInferenceQAAsset->CompareNewVsPreviousTests(UnitTestDirectory))
 	{
 		// Saving NeuralNetworkInferenceQAAsset
 		ensureMsgf(NeuralNetworkInferenceQAAsset->Save(), TEXT("NeuralNetworkInferenceQAAsset->Save() failed."));
diff --git a/Engine/Plugins/Experimental/NeuralNetworkInference/Source/NeuralNetworkInferenceQA/Private/OperatorUnitTester.h b/Engine/Plugins/Experimental/NeuralNetworkInference/Source/NeuralNetworkInferenceQA/Private/OperatorUnitTester.h
--- a/Engine/Plugins/Experimental/NeuralNetworkInference/Source/NeuralNetworkInferenceQA/Private/OperatorUnitTester.h
+++ b/Engine/Plugins/Experimental/NeuralNetworkInference/Source/NeuralNetworkInferenceQA/Private/OperatorUnitTester.h
@@ -8,7 +8,7 @@
 class FOperatorUnitTester
 {
 public:
-	static void GlobalTest(const FString& InGroundTruthDirectory);
+	static void GlobalTest(const FString& InProjectContentDir, const FString& InUnitTestRelativeDirectory);
 
 private:
 	static void CreateRandomArraysOfFixedSeedAndDefinedSizes(TMap<int32, TArray<float>>& InOutInputTensorData, TMap<int32, TArray<float>>& InOutInputTensorDataSmall,
diff --git a/Engine/Plugins/Experimental/NeuralNetworkInference/Source/NeuralNetworkInferenceQA/Private/UnitTester.h b/Engine/Plugins/Experimental/NeuralNetworkInference/Source/NeuralNetworkInferenceQA/Private/UnitTester.h
--- a/Engine/Plugins/Experimental/NeuralNetworkInference/Source/NeuralNetworkInferenceQA/Private/UnitTester.h
+++ b/Engine/Plugins/Experimental/NeuralNetworkInference/Source/NeuralNetworkInferenceQA/Private/UnitTester.h
@@ -10,5 +10,5 @@ public:
 	/**
 	 * Main unit tester function.
 	 */
-	static void GlobalTest(const FString& InGroundTruthDirectory, const FString& InModelsDirectory);
+	static void GlobalTest(const FString& InProjectContentDir, const FString& InModelZooRelativeDirectory, const FString& InUnitTestRelativeDirectory);
 };
