NeuralNetworkInference (NNI) plugin:
- Preparing NNIQA for automation tests. (Part 2)

#rnx

diff --git a/Engine/Plugins/Experimental/NeuralNetworkInference/Source/NeuralNetworkInferenceQA/Private/LegacyModelUnitTester.cpp b/Engine/Plugins/Experimental/NeuralNetworkInference/Source/NeuralNetworkInferenceQA/Private/LegacyModelUnitTester.cpp
--- a/Engine/Plugins/Experimental/NeuralNetworkInference/Source/NeuralNetworkInferenceQA/Private/LegacyModelUnitTester.cpp
+++ b/Engine/Plugins/Experimental/NeuralNetworkInference/Source/NeuralNetworkInferenceQA/Private/LegacyModelUnitTester.cpp
@@ -11,7 +11,7 @@
 /* FLegacyModelUnitTester static public functions
  *****************************************************************************/
 
-void FLegacyModelUnitTester::GlobalTest(const FString& InModelsDirectory)
+void FLegacyModelUnitTester::GlobalTest(const FString& InProjectContentDir, const FString& InModelZooRelativeDirectory)
 {
 	// Model load, accuracy, and speed test
 	const TArray<FString> ModelNames({ TEXT("MLRigDeformer")/*, TEXT("cloth_network")*/ });
@@ -20,7 +20,7 @@ void FLegacyModelUnitTester::GlobalTest(const FString& InModelsDirectory)
 	const TArray<TArray<double>> GPUGroundTruths({ {3.728547, 0.008774, 4.595651, 212.193208, 742.434578, 4.250668, 4.717748}, {0.042571, 0.023693, 0.015783, 13.100504, 8.050994, 0.028807, 0.016387} });
 	const TArray<int32> CPURepetitions({ 10, 0 }); // Used to be 100
 	const TArray<int32> GPURepetitions({ 10, 10 }); // Used to be 100
-	ModelLoadAccuracyAndSpeedTests(InModelsDirectory, ModelNames, InputArrayValues, CPUGroundTruths, GPUGroundTruths, CPURepetitions, GPURepetitions);
+	ModelLoadAccuracyAndSpeedTests(InProjectContentDir, InModelZooRelativeDirectory, ModelNames, InputArrayValues, CPUGroundTruths, GPUGroundTruths, CPURepetitions, GPURepetitions);
 }
 
 
@@ -28,17 +28,20 @@ void FLegacyModelUnitTester::GlobalTest(const FString& InModelsDirectory)
 /* FLegacyModelUnitTester static private functions
  *****************************************************************************/
 
-void FLegacyModelUnitTester::ModelLoadAccuracyAndSpeedTests(const FString& InModelsDirectory, const TArray<FString>& InModelNames, const TArray<float>& InInputArrayValues,
-	const TArray<TArray<double>>& InCPUGroundTruths, const TArray<TArray<double>>& InGPUGroundTruths, const TArray<int32>& InCPURepetitions, const TArray<int32>& InGPURepetitions)
+void FLegacyModelUnitTester::ModelLoadAccuracyAndSpeedTests(const FString& InProjectContentDir, const FString& InModelZooRelativeDirectory, const TArray<FString>& InModelNames,
+	const TArray<float>& InInputArrayValues, const TArray<TArray<double>>& InCPUGroundTruths, const TArray<TArray<double>>& InGPUGroundTruths, const TArray<int32>& InCPURepetitions,
+	const TArray<int32>& InGPURepetitions)
 {
+	const FString ModelZooDirectory = InProjectContentDir / InModelZooRelativeDirectory;
+
 	// Test ONNX Reader
 	for (const FString& ModelName : InModelNames)
 	{
 		UE_LOG(LogNeuralNetworkInferenceQA, Display, TEXT("--------------- %s - OTXT IO"), *ModelName);
-		const FString OtxtModelFilePath = GetOtxtModelFilePath(InModelsDirectory, ModelName);
+		const FString OtxtModelFilePath = GetOtxtModelFilePath(ModelZooDirectory, ModelName);
 		FLegacyONNXTester::ONNXReadNetworkTest(OtxtModelFilePath);
 		UE_LOG(LogNeuralNetworkInferenceQA, Display, TEXT("--------------- %s - ONNX IO"), *ModelName);
-		const FString ONNXModelFilePath = GetONNXModelFilePath(InModelsDirectory, ModelName);
+		const FString ONNXModelFilePath = GetONNXModelFilePath(ModelZooDirectory, ModelName);
 		FLegacyONNXTester::ONNXReadNetworkTest(ONNXModelFilePath);
 	}
 	UE_LOG(LogNeuralNetworkInferenceQA, Display, TEXT("---------------------------------------------------------------------------------------------------------------------------------"));
@@ -50,12 +53,12 @@ void FLegacyModelUnitTester::ModelLoadAccuracyAndSpeedTests(const FString& InMod
 		const TArray<double>& CPUGroundTruths = InCPUGroundTruths[ModelIndex];
 		const TArray<double>& GPUGroundTruths = InGPUGroundTruths[ModelIndex];
 
-		UE_LOG(LogNeuralNetworkInferenceQA, Display, TEXT("--------------- %s - Legacy Network ONNX Load and Run"), *ModelName);
-		const FString ONNXModelFilePath = GetONNXModelFilePath(InModelsDirectory, ModelName);
+		const FString ONNXModelFilePath = GetONNXModelFilePath(ModelZooDirectory, ModelName);
+		UE_LOG(LogNeuralNetworkInferenceQA, Display, TEXT("--------------- %s - Legacy Network ONNX Load and Run - %s"), *ModelName, *ONNXModelFilePath);
 		ModelAccuracyTest(NetworkONNXLoadTest(ONNXModelFilePath), InInputArrayValues, CPUGroundTruths, GPUGroundTruths);
 
 		UE_LOG(LogNeuralNetworkInferenceQA, Display, TEXT("--------------- %s - Legacy Network Uasset Load and Run"), *ModelName);
-		const FString UAssetModelFilePath = GetUAssetModelFilePath(ModelName);
+		const FString UAssetModelFilePath = GetUAssetModelFilePath(ModelName, InModelZooRelativeDirectory);
 		ModelAccuracyTest(NetworkUassetLoadTest(UAssetModelFilePath), InInputArrayValues, CPUGroundTruths, GPUGroundTruths);
 
 		UE_LOG(LogNeuralNetworkInferenceQA, Display, TEXT("---------------------------------------------------------------------------------------------------------------------------------"));
@@ -66,24 +69,24 @@ void FLegacyModelUnitTester::ModelLoadAccuracyAndSpeedTests(const FString& InMod
 	{
 		const FString& ModelName = InModelNames[ModelIndex];
 		UE_LOG(LogNeuralNetworkInferenceQA, Display, TEXT("--------------- %s - Legacy Network UAsset Speed Profiling"), *ModelName);
-		const FString UAssetModelFilePath = GetUAssetModelFilePath(ModelName);
+		const FString UAssetModelFilePath = GetUAssetModelFilePath(ModelName, InModelZooRelativeDirectory);
 		ModelSpeedTest(UAssetModelFilePath, InCPURepetitions[ModelIndex], InGPURepetitions[ModelIndex]);
 	}
 }
 
-FString FLegacyModelUnitTester::GetONNXModelFilePath(const FString& InModelsDirectory, const FString& InModelName)
+FString FLegacyModelUnitTester::GetONNXModelFilePath(const FString& ModelZooDirectory, const FString& InModelName)
 {
-	return FPaths::ConvertRelativePathToFull(InModelsDirectory + InModelName + TEXT("/") + InModelName + TEXT(".onnx")); // E.g., InModelsDirectory + TEXT("ExampleNetworkReadable/ExampleNetworkReadable.onnx")
+	return FPaths::ConvertRelativePathToFull(ModelZooDirectory + InModelName + TEXT("/") + InModelName + TEXT(".onnx")); // E.g., ModelZooDirectory / TEXT("ExampleNetworkReadable/ExampleNetworkReadable.onnx")
 }
 
-FString FLegacyModelUnitTester::GetOtxtModelFilePath(const FString& InModelsDirectory, const FString& InModelName)
+FString FLegacyModelUnitTester::GetOtxtModelFilePath(const FString& ModelZooDirectory, const FString& InModelName)
 {
-	return FPaths::ConvertRelativePathToFull(InModelsDirectory + InModelName + TEXT("/") + InModelName + TEXT(".otxt")); // E.g., InModelsDirectory + TEXT("ExampleNetworkReadable/ExampleNetworkReadable.otxt")
+	return FPaths::ConvertRelativePathToFull(ModelZooDirectory + InModelName + TEXT("/") + InModelName + TEXT(".otxt")); // E.g., ModelZooDirectory / TEXT("ExampleNetworkReadable/ExampleNetworkReadable.otxt")
 }
 
-FString FLegacyModelUnitTester::GetUAssetModelFilePath(const FString& InModelName)
+FString FLegacyModelUnitTester::GetUAssetModelFilePath(const FString& InModelName, const FString& InModelZooRelativeDirectory)
 {
-	return InModelName + TEXT("'/Game/Models/") + InModelName + TEXT("/") + InModelName + TEXT("Legacy.") + InModelName + TEXT("Legacy'"); // ExampleNetworkReadable'/Game/Models/ExampleNetworkReadable/ExampleNetworkReadable.ExampleNetworkReadable'
+	return InModelName + TEXT("'/Game/") + InModelZooRelativeDirectory / InModelName + TEXT("/") + InModelName + TEXT("Legacy.") + InModelName + TEXT("Legacy'"); // E.g., '/Game/[MODEL_ZOO_DIR]/ExampleNetworkReadable/ExampleNetworkReadable.ExampleNetworkReadable'
 }
 
 UNeuralNetworkLegacy* FLegacyModelUnitTester::NetworkUassetLoadTest(const FString& InUAssetPath)
diff --git a/Engine/Plugins/Experimental/NeuralNetworkInference/Source/NeuralNetworkInferenceQA/Private/UnitTester.cpp b/Engine/Plugins/Experimental/NeuralNetworkInference/Source/NeuralNetworkInferenceQA/Private/UnitTester.cpp
--- a/Engine/Plugins/Experimental/NeuralNetworkInference/Source/NeuralNetworkInferenceQA/Private/UnitTester.cpp
+++ b/Engine/Plugins/Experimental/NeuralNetworkInference/Source/NeuralNetworkInferenceQA/Private/UnitTester.cpp
@@ -17,7 +17,7 @@
 /* FUnitTester static public functions
  *****************************************************************************/
 
-void FUnitTester::GlobalTest(const FString& InGroundTruthDirectory, const FString& InModelsDirectory)
+void FUnitTester::GlobalTest(const FString& InProjectContentDir, const FString& InModelZooRelativeDirectory, const FString& InUnitTestRelativeDirectory)
 {
 	UE_LOG(LogNeuralNetworkInferenceQA, Display, TEXT("---------------------------------------------------------------------------------------------------------------------------------"));
 	UE_LOG(LogNeuralNetworkInferenceQA, Display, TEXT("---------------------------------------------------------------------------------------------------------------------------------"));
@@ -25,19 +25,19 @@ void FUnitTester::GlobalTest(const FString& InGroundTruthDirectory, const FStrin
 	UE_LOG(LogNeuralNetworkInferenceQA, Display, TEXT("----- Starting UnitTesting() ----------------------------------------------------------------------------------------------------"));
 
 	UE_LOG(LogNeuralNetworkInferenceQA, Display, TEXT("---------- 1. Model Unit Testing (Legacy)"));
-	FLegacyModelUnitTester::GlobalTest(InModelsDirectory);
+	FLegacyModelUnitTester::GlobalTest(InProjectContentDir, InModelZooRelativeDirectory);
 
 	UE_LOG(LogNeuralNetworkInferenceQA, Display, TEXT("---------------------------------------------------------------------------------------------------------------------------------"));
 	UE_LOG(LogNeuralNetworkInferenceQA, Display, TEXT("---------------------------------------------------------------------------------------------------------------------------------"));
 	UE_LOG(LogNeuralNetworkInferenceQA, Display, TEXT("---------------------------------------------------------------------------------------------------------------------------------"));
 	UE_LOG(LogNeuralNetworkInferenceQA, Display, TEXT("---------- 2. Model Unit Testing"));
-	FModelUnitTester::GlobalTest(InModelsDirectory);
+	FModelUnitTester::GlobalTest(InProjectContentDir, InModelZooRelativeDirectory);
 
 	UE_LOG(LogNeuralNetworkInferenceQA, Display, TEXT("---------------------------------------------------------------------------------------------------------------------------------"));
 	UE_LOG(LogNeuralNetworkInferenceQA, Display, TEXT("---------------------------------------------------------------------------------------------------------------------------------"));
 	UE_LOG(LogNeuralNetworkInferenceQA, Display, TEXT("---------------------------------------------------------------------------------------------------------------------------------"));
 	UE_LOG(LogNeuralNetworkInferenceQA, Display, TEXT("---------- 3. Operator Unit Testing"));
-	FOperatorUnitTester::GlobalTest(InGroundTruthDirectory);
+	FOperatorUnitTester::GlobalTest(InProjectContentDir, InUnitTestRelativeDirectory);
 
 // #if WITH_EDITOR
 // #ifdef PLATFORM_WIN64
@@ -45,7 +45,7 @@ void FUnitTester::GlobalTest(const FString& InGroundTruthDirectory, const FStrin
 // 	UE_LOG(LogNeuralNetworkInferenceQA, Display, TEXT("---------------------------------------------------------------------------------------------------------------------------------"));
 // 	UE_LOG(LogNeuralNetworkInferenceQA, Display, TEXT("---------------------------------------------------------------------------------------------------------------------------------"));
 // 	UE_LOG(LogNeuralNetworkInferenceQA, Display, TEXT("---------- 4. ONNX Runtime DLL Unit Testing (Deprecated)"));
-// 	FONNXRuntimeDLLTester::GlobalTest(InModelsDirectory);
+// 	FONNXRuntimeDLLTester::GlobalTest(InProjectContentDir, InModelZooRelativeDirectory);
 // #endif //PLATFORM_WIN64
 // #endif //WITH_EDITOR
 
