NeuralNetworkInference (NNI) plugin:
- Added Automation tests for NNI.

#rnx

diff --git a/Engine/Plugins/Experimental/NeuralNetworkInference/Source/NeuralNetworkInferenceQA/Private/NeuralNetworkInferenceQA.cpp b/Engine/Plugins/Experimental/NeuralNetworkInference/Source/NeuralNetworkInferenceQA/Private/NeuralNetworkInferenceQA.cpp
--- a/Engine/Plugins/Experimental/NeuralNetworkInference/Source/NeuralNetworkInferenceQA/Private/NeuralNetworkInferenceQA.cpp
+++ b/Engine/Plugins/Experimental/NeuralNetworkInference/Source/NeuralNetworkInferenceQA/Private/NeuralNetworkInferenceQA.cpp
@@ -18,3 +18,25 @@ void UNeuralNetworkInferenceQA::UnitTesting()
 	const FString UnitTestRelativeDirectory = MachineLearningTestsRelativeDirectory / TEXT("UnitTests/"); // Eg "Tests/MachineLearning/UnitTests/"
 	FUnitTester::GlobalTest(ProjectContentDir, ModelZooRelativeDirectory, UnitTestRelativeDirectory);
 }
+
+
+
+#if WITH_DEV_AUTOMATION_TESTS
+
+#include "Misc/AutomationTest.h"
+
+IMPLEMENT_SIMPLE_AUTOMATION_TEST(FNeuralNetworkInferenceTest, "System.Engine.MachineLearning.NeuralNetworkInference (NNI)",
+	EAutomationTestFlags::ApplicationContextMask // = EditorContext | ClientContext | ServerContext | CommandletContext
+	| EAutomationTestFlags::FeatureMask // = NonNullRHI | RequiresUser
+	| EAutomationTestFlags::EngineFilter)
+
+bool FNeuralNetworkInferenceTest::RunTest(const FString& Parameters)
+{
+	UNeuralNetworkInferenceQA::UnitTesting();
+UE_LOG(LogNeuralNetworkInferenceQA, Warning, TEXT("FNeuralNetworkInferenceTest::RunTest(): Warning with parameters = %s."), *Parameters);
+ensureMsgf(false, TEXT("FNeuralNetworkInferenceTest::RunTest(): Failure with parameters = %s."), *Parameters);
+	return true;
+}
+
+
+#endif //WITH_DEV_AUTOMATION_TESTS
