NeuralNetworkInference (NNI) plugin:
- UAsset renames.

#rnx

diff --git a/Engine/Plugins/Experimental/NeuralNetworkInference/Source/NeuralNetworkInferenceQA/Private/ModelUnitTester.cpp b/Engine/Plugins/Experimental/NeuralNetworkInference/Source/NeuralNetworkInferenceQA/Private/ModelUnitTester.cpp
--- a/Engine/Plugins/Experimental/NeuralNetworkInference/Source/NeuralNetworkInferenceQA/Private/ModelUnitTester.cpp
+++ b/Engine/Plugins/Experimental/NeuralNetworkInference/Source/NeuralNetworkInferenceQA/Private/ModelUnitTester.cpp
@@ -14,7 +14,7 @@ void FModelUnitTester::GlobalTest(const FString& InProjectContentDir, const FStr
 {
 #ifdef WITH_FULL_NNI_SUPPORT
 	// Model load, accuracy, and speed test
-	const TArray<FString> ModelNames({ TEXT("MLRigDeformer"), TEXT("cloth_network"), TEXT("hyprsense"), TEXT("rocket_league") });
+	const TArray<FString> ModelNames({ TEXT("MLRigDeformer"), TEXT("cloth_network"), TEXT("HS"), TEXT("RL") });
 	const TArray<float> InputArrayValues({ 1.f, 0.f, -1.f, 100.f, -100.f, 0.5f, -0.5f }); // This one can be shorter than CPU/GPUGroundTruths
 	const TArray<TArray<double>> CPUGroundTruths({ {3.728547, 0.008774, 4.595651, 212.193216, 742.434561, 4.250668, 4.717748}, {0.042571, 0.023693, 0.015783, 13.100505, 8.050994, 0.028807, 0.016387},
 		{138.372906, 126.753839, 127.287254, 130.316062, 127.303424, 124.800896, 126.546051}, {0.488662, 0.472437, 0.478862, 0.522685, 0.038322, 0.480848, 0.483821}});
