NeuralNetworkInference (NNI) plugin:
- Fixed Automation tests for NNI.

#rnx

diff --git a/Engine/Plugins/Experimental/NeuralNetworkInference/Source/NeuralNetworkInferenceQA/Private/NeuralNetworkInferenceQA.cpp b/Engine/Plugins/Experimental/NeuralNetworkInference/Source/NeuralNetworkInferenceQA/Private/NeuralNetworkInferenceQA.cpp
--- a/Engine/Plugins/Experimental/NeuralNetworkInference/Source/NeuralNetworkInferenceQA/Private/NeuralNetworkInferenceQA.cpp
+++ b/Engine/Plugins/Experimental/NeuralNetworkInference/Source/NeuralNetworkInferenceQA/Private/NeuralNetworkInferenceQA.cpp
@@ -34,7 +34,6 @@ bool FNeuralNetworkInferenceTest::RunTest(const FString& Parameters)
 {
 	UNeuralNetworkInferenceQA::UnitTesting();
 UE_LOG(LogNeuralNetworkInferenceQA, Warning, TEXT("FNeuralNetworkInferenceTest::RunTest(): Warning with parameters = %s."), *Parameters);
-ensureMsgf(false, TEXT("FNeuralNetworkInferenceTest::RunTest(): Failure with parameters = %s."), *Parameters);
 	return true;
 }
 
