NeuralNetworkInference (NNI) plugin:
- Prepared API and UAsset for (a)sync Run(), but not implemented yet.

#rnx

diff --git a/Engine/Plugins/Experimental/NeuralNetworkInference/Source/NeuralNetworkInference/Private/NeuralNetwork.cpp b/Engine/Plugins/Experimental/NeuralNetworkInference/Source/NeuralNetworkInference/Private/NeuralNetwork.cpp
--- a/Engine/Plugins/Experimental/NeuralNetworkInference/Source/NeuralNetworkInference/Private/NeuralNetwork.cpp
+++ b/Engine/Plugins/Experimental/NeuralNetworkInference/Source/NeuralNetworkInference/Private/NeuralNetwork.cpp
@@ -232,6 +232,7 @@ UNeuralNetwork::UNeuralNetwork()
 	: DeviceType(ENeuralDeviceType::CPU)
 	, InputDeviceType(ENeuralDeviceType::CPU)
 	, OutputDeviceType(ENeuralDeviceType::CPU)
+	, SynchronousMode(ENeuralNetworkSynchronousMode::Synchronous)
 	, bIsLoaded(false)
 {
 }
@@ -473,10 +474,21 @@ void UNeuralNetwork::Run()
 	}
 
 	// Run UNeuralNetwork
-	const FRedirectCoutAndCerrToUeLog RedirectCoutAndCerrToUeLog;
-	Impl->Session->Run(Ort::RunOptions{ nullptr },
-		InputTensors.GetTensorNames(), InputTensors.GetONNXRuntimeTensors(), InputTensors.GetNumberTensors(),
-		OutputTensors.GetTensorNames(), OutputTensors.GetONNXRuntimeTensors(), OutputTensors.GetNumberTensors());
+	if (SynchronousMode == ENeuralNetworkSynchronousMode::Synchronous)
+	{
+		const FRedirectCoutAndCerrToUeLog RedirectCoutAndCerrToUeLog;
+		Impl->Session->Run(Ort::RunOptions{ nullptr },
+			InputTensors.GetTensorNames(), InputTensors.GetONNXRuntimeTensors(), InputTensors.GetNumberTensors(),
+			OutputTensors.GetTensorNames(), OutputTensors.GetONNXRuntimeTensors(), OutputTensors.GetNumberTensors());
+	}
+	else if (SynchronousMode == ENeuralNetworkSynchronousMode::Asynchronous)
+	{
+		UE_LOG(LogNeuralNetworkInference, Warning, TEXT("UNeuralNetwork::Run(): SynchronousMode = %d not implemented yet. Use SynchronousMode = Synchronous."), (int32)SynchronousMode);
+	}
+	else
+	{
+		UE_LOG(LogNeuralNetworkInference, Warning, TEXT("UNeuralNetwork::Run(): Unknown SynchronousMode = %d."), (int32)SynchronousMode);
+	}
 
 #else
 	UE_LOG(LogNeuralNetworkInference, Warning, TEXT("UNeuralNetwork::Run(): Platform or Operating System not suported yet."));
diff --git a/Engine/Plugins/Experimental/NeuralNetworkInference/Source/NeuralNetworkInference/Public/NeuralNetwork.h b/Engine/Plugins/Experimental/NeuralNetworkInference/Source/NeuralNetworkInference/Public/NeuralNetwork.h
--- a/Engine/Plugins/Experimental/NeuralNetworkInference/Source/NeuralNetworkInference/Public/NeuralNetwork.h
+++ b/Engine/Plugins/Experimental/NeuralNetworkInference/Source/NeuralNetworkInference/Public/NeuralNetwork.h
@@ -10,7 +10,7 @@
 class UAssetImportData;
 
 /**
- * Whether Run() will block the thread until completed, or whether it will run on a background thread, not blocking the calling thread.
+ * Whether UNeuralNetwork::Run() will block the thread until completed (Synchronous), or whether it will run on a background thread, not blocking the calling thread (Asynchronous).
  */
 UENUM()
 enum class ENeuralNetworkSynchronousMode : uint8
@@ -148,6 +148,13 @@ protected:
 	
 	UPROPERTY(VisibleAnywhere, Category = "Neural Network Inference")
 	ENeuralDeviceType OutputDeviceType;
+	
+	/**
+	 * Whether UNeuralNetwork::Run() will block the thread until completed (Synchronous), or whether it will run on a background thread, not blocking the calling thread (Asynchronous).
+	 * See ENeuralNetworkSynchronousMode for more details.
+	 */
+	UPROPERTY(VisibleAnywhere, Category = "Neural Network Inference")
+	ENeuralNetworkSynchronousMode SynchronousMode;
 
 	/**
 	 * Original model file path from which this UNeuralNetwork was loaded from.
