NeuralNetworkInference (NNI) plugin:
- NNILegacy merged into NNI step 1: UNN & UNNLegacy APIs more similar to each other.

#rnx

diff --git a/Engine/Plugins/Experimental/NeuralNetworkInference/Source/NeuralNetworkInference/Internal/NeuralNetworkLegacy.h b/Engine/Plugins/Experimental/NeuralNetworkInference/Source/NeuralNetworkInference/Internal/NeuralNetworkLegacy.h
--- a/Engine/Plugins/Experimental/NeuralNetworkInference/Source/NeuralNetworkInference/Internal/NeuralNetworkLegacy.h
+++ b/Engine/Plugins/Experimental/NeuralNetworkInference/Source/NeuralNetworkInference/Internal/NeuralNetworkLegacy.h
@@ -3,14 +3,17 @@
 #pragma once
 
 #include "CoreMinimal.h"
-#include "ModelProto.h"
 #include "NeuralEnumClasses.h"
-#include "NeuralNetwork.h"
-#include "NeuralOperator.h"
 #include "NeuralTensor.h"
+#include "NeuralTensors.h"
+
+#include "NeuralNetwork.h"
+#include "ModelProto.h"
 #include "NeuralTensorManager.h"
+
 #include "NeuralNetworkLegacy.generated.h"
 
+class FNeuralOperator;
 class UAssetImportData;
 
 /**
@@ -22,60 +25,42 @@ class NEURALNETWORKINFERENCE_API UNeuralNetworkLegacy : public UObject
 	GENERATED_BODY()
 
 public:
-	/**
-	 * Callbacks when the asyn Run() is finished. It will happen from any thread.
-	 */
-	DECLARE_DELEGATE(FOnAsyncRunCompletedInAnyThread);
-
 	UNeuralNetworkLegacy();
 
 	virtual ~UNeuralNetworkLegacy();
 
-	//~UObject interface
-	virtual void PostInitProperties() override;
-	virtual void PostLoad() override;
-	virtual void Serialize(FArchive& Archive) override;
-	//~End of UObject interface
-#if WITH_EDITOR
-	/** Editor-only function: Re-import asset with editor data (imported file). */
-	void ReimportAssetFromEditorData();
-	/** Editor-only function: Importing data and options used for loading the neural network. */
-	UAssetImportData* GetAssetImportData() const;
-	UAssetImportData* GetAndMaybeCreateAssetImportData();
-#endif // WITH_EDITOR
-
 	/**
 	 * Editor-only function
-	 * It loads the desired UNeuralNetworkLegacy definition and its weights from an ONNX file.
-	 * @return Whether it loaded the UNeuralNetworkLegacy successfully.
+	 * It loads the desired network graph definition and weights given an input ONNX file path.
+	 * @param InModelFilePath can either be a full path or a relative path with respect to the Game project.
+	 * @return Whether the network was successfully loaded.
 	 */
 #if WITH_EDITOR
-	bool Load(const FString& InFilePath);
+	bool Load(const FString& InModelFilePath);
 #endif // WITH_EDITOR
 
 	/**
-	 * It loads the desired UNeuralNetworkLegacy definition and its weights from an deserialized UNeuralNetworkLegacy uasset.
+	 * It loads the desired network graph definition and weights given an input UNeuralNetworkLegacy UAsset.
 	 * @param InTensorManager It will be moved for performance reasons. Do not use InTensorManager after calling this function.
-	 * @return Whether it loaded the UNeuralNetworkLegacy successfully.
+	 * @return Whether the network was successfully loaded.
 	 */
 	bool Load(FNeuralTensorManager& InTensorManager, const TArray<TSharedPtr<FNeuralOperator>>& InOperators);
 
 	/**
-	 * It returns whether it loaded the UNeuralNetworkLegacy successfully.
+	 * It returns whether a network is currently loaded.
 	 */
 	bool IsLoaded() const;
 
 	/**
-	 * Whether Run() will occur on the CPU or GPU.
-	 * If SetDeviceType() is never called, the default device (EDeviceType::CPU) will be used.
+	 * Getter and setter functions for DeviceType, InputDeviceType, and OutputDeviceType.
+	 * See DeviceType, InputDeviceType, and OutputDeviceType for more details.
 	 */
 	ENeuralDeviceType GetDeviceType() const;
 	void SetDeviceType(const ENeuralDeviceType InDeviceType);
-
-	/**
-	 * Returns the delegate called when async Run() is over (on any thread)
-	 */
-	FOnAsyncRunCompletedInAnyThread& GetOnAsyncRunCompletedInAnyThreadDelegate();
+	ENeuralDeviceType GetInputDeviceType() const;
+	void SetInputDeviceType(const ENeuralDeviceType InInputDeviceType);
+	ENeuralDeviceType GetOutputDeviceType() const;
+	void SetOutputDeviceType(const ENeuralDeviceType InOutputDeviceType);
 
 	/**
 	 * It returns the TArray of FNeuralTensors as a read-only object.
@@ -144,6 +129,16 @@ public:
 	TMap<FString, FNeuralTensor> CreateInputTensorMap() const;
 	TMap<FString, FNeuralTensor> CreateOutputTensorMap() const;
 
+	/**
+	 * Callback when UNeuralNetwork::Run() is finished. It will only occur if SynchronousMode == ENeuralNetworkSynchronousMode::Asynchronous and it could be triggered from any thread.
+	 */
+	DECLARE_DELEGATE(FOnAsyncRunCompleted);
+
+	/**
+	 * Returns the delegate called when async Run() is over (on any thread)
+	 */
+	FOnAsyncRunCompleted& GetOnAsyncRunCompletedDelegate();
+
 	/**
 	 * Run() executes the forward pass on the current UNeuralNetworkLegacy given the current input FNeuralTensor(s), which were previously filled with
 	 * SetInputFromTensorCopy() or GetInputDataPointerMutable(), or their multi-input analogs.
@@ -163,8 +158,26 @@ protected:
 	TArray<int32> Version;
 	UPROPERTY(VisibleAnywhere, Category = "Neural Network Inference")
 	bool bIsLoaded;
-	UPROPERTY(EditAnywhere, Category = "Neural Network Inference")
+	
+	/**
+	 * Whether Run() will use CPU or GPU acceleration hardware.
+	 * If SetDeviceType() is never called, the default device (EDeviceType::CPU) will be used.
+	 */
+	UPROPERTY(VisibleAnywhere, Category = "Neural Network Inference")
 	ENeuralDeviceType DeviceType;
+	
+	/**
+	 * If DeviceType == CPU, InputDeviceType and OutputDeviceType must also be set to CPU.
+	 * If DeviceType == GPU:
+	 *  - InputDeviceType defines whether Run() will expect the input data in CPU (Run() will upload the memory to the GPU first) or GPU (no upload copy needed) format.
+	 *  - OutputDeviceType defines whether Run() will return output data in CPU (Run() will download the memory to the CPU first) or GPU (no download copy needed) format.
+	 */
+	UPROPERTY(VisibleAnywhere, Category = "Neural Network Inference")
+	ENeuralDeviceType InputDeviceType;
+	
+	UPROPERTY(VisibleAnywhere, Category = "Neural Network Inference")
+	ENeuralDeviceType OutputDeviceType;
+
 	UPROPERTY(VisibleAnywhere, Category = "Neural Network Inference")
 	FNeuralTensorManager TensorManager; /* It contains a few TArray and TMaps for all FNeuralTensors (Input, Output, Intermediate(Not)Initialized, Weight) */
 	UPROPERTY(VisibleAnywhere, Category = "Neural Network Inference")
@@ -182,7 +195,21 @@ private:
 	 * Operators represents the set of operators of the network that need to run on the Forward pass and that might need to run on the PostForward.
 	 */
 	TArray<TSharedPtr<FNeuralOperator>> Operators;
-	FOnAsyncRunCompletedInAnyThread OnAsyncRunCompletedInAnyThreadDelegate;
+	FOnAsyncRunCompleted OnAsyncRunCompletedDelegate;
+
+public:
+	//~UObject interface
+	virtual void PostInitProperties() override;
+	virtual void PostLoad() override;
+	virtual void Serialize(FArchive& Archive) override;
+	//~End of UObject interface
+#if WITH_EDITOR
+/** Editor-only function: Re-import asset with editor data (imported file). */
+	void ReimportAssetFromEditorData();
+	/** Editor-only function: Importing data and options used for loading the neural network. */
+	UAssetImportData* GetAssetImportData() const;
+	UAssetImportData* GetAndMaybeCreateAssetImportData();
+#endif // WITH_EDITOR
 };
 
 
diff --git a/Engine/Plugins/Experimental/NeuralNetworkInference/Source/NeuralNetworkInference/Private/NeuralNetwork.cpp b/Engine/Plugins/Experimental/NeuralNetworkInference/Source/NeuralNetworkInference/Private/NeuralNetwork.cpp
--- a/Engine/Plugins/Experimental/NeuralNetworkInference/Source/NeuralNetworkInference/Private/NeuralNetwork.cpp
+++ b/Engine/Plugins/Experimental/NeuralNetworkInference/Source/NeuralNetworkInference/Private/NeuralNetwork.cpp
@@ -229,11 +229,11 @@ void UNeuralNetwork::FImpl::ConfigureTensors(FNeuralTensors& OutTensors, TArray<
  *****************************************************************************/
 
 UNeuralNetwork::UNeuralNetwork()
-	: DeviceType(ENeuralDeviceType::CPU)
+	: bIsLoaded(false)
+	, DeviceType(ENeuralDeviceType::CPU)
 	, InputDeviceType(ENeuralDeviceType::CPU)
 	, OutputDeviceType(ENeuralDeviceType::CPU)
 	, SynchronousMode(ENeuralNetworkSynchronousMode::Synchronous)
-	, bIsLoaded(false)
 {
 }
 
@@ -457,6 +457,11 @@ const FNeuralTensors& UNeuralNetwork::GetOutputTensors() const
 	return OutputTensors;
 }
 
+UNeuralNetwork::FOnAsyncRunCompleted& UNeuralNetwork::GetOnAsyncRunCompletedDelegate()
+{
+	return OnAsyncRunCompletedDelegate;
+}
+
 void UNeuralNetwork::Run()
 {
 #ifdef WITH_FULL_NNI_SUPPORT
diff --git a/Engine/Plugins/Experimental/NeuralNetworkInference/Source/NeuralNetworkInference/Private/NeuralNetworkLegacy.cpp b/Engine/Plugins/Experimental/NeuralNetworkInference/Source/NeuralNetworkInference/Private/NeuralNetworkLegacy.cpp
--- a/Engine/Plugins/Experimental/NeuralNetworkInference/Source/NeuralNetworkInference/Private/NeuralNetworkLegacy.cpp
+++ b/Engine/Plugins/Experimental/NeuralNetworkInference/Source/NeuralNetworkInference/Private/NeuralNetworkLegacy.cpp
@@ -202,9 +202,24 @@ void UNeuralNetworkLegacy::SetDeviceType(const ENeuralDeviceType InDeviceType)
 	DeviceType = InDeviceType;
 }
 
-UNeuralNetworkLegacy::FOnAsyncRunCompletedInAnyThread& UNeuralNetworkLegacy::GetOnAsyncRunCompletedInAnyThreadDelegate()
+ENeuralDeviceType UNeuralNetworkLegacy::GetInputDeviceType() const
 {
-	return OnAsyncRunCompletedInAnyThreadDelegate;
+	return InputDeviceType;
+}
+
+void UNeuralNetworkLegacy::SetInputDeviceType(const ENeuralDeviceType InInputDeviceType)
+{
+	InputDeviceType = InInputDeviceType;
+}
+
+ENeuralDeviceType UNeuralNetworkLegacy::GetOutputDeviceType() const
+{
+	return OutputDeviceType;
+}
+
+void UNeuralNetworkLegacy::SetOutputDeviceType(const ENeuralDeviceType InOutputDeviceType)
+{
+	OutputDeviceType = InOutputDeviceType;
 }
 
 const TArray<FNeuralTensor>& UNeuralNetworkLegacy::GetTensors() const
@@ -277,6 +292,11 @@ TMap<FString, FNeuralTensor> UNeuralNetworkLegacy::CreateOutputTensorMap() const
 	return TensorManager.CreateOutputTensorMap();
 }
 
+UNeuralNetworkLegacy::FOnAsyncRunCompleted& UNeuralNetworkLegacy::GetOnAsyncRunCompletedDelegate()
+{
+	return OnAsyncRunCompletedDelegate;
+}
+
 void UNeuralNetworkLegacy::Run(const ENeuralNetworkSynchronousMode InSynchronousMode, const ENeuralDeviceType InInputDeviceType, const ENeuralDeviceType InOutputDeviceType, const bool bRunGPUEmptyOnlyForProfiling)
 {
 	DECLARE_SCOPE_CYCLE_COUNTER(TEXT("UNeuralNetworkLegacy_Run"), STAT_UNeuralNetworkLegacy_Run, STATGROUP_MachineLearning);
@@ -415,7 +435,7 @@ void UNeuralNetworkLegacy::Run(const ENeuralNetworkSynchronousMode InSynchronous
 							ERDGPassFlags::None,
 							[this](FRHICommandListImmediate& RHICmdList)
 						{
-							OnAsyncRunCompletedInAnyThreadDelegate.ExecuteIfBound();
+							OnAsyncRunCompletedDelegate.ExecuteIfBound();
 						});
 					}
 
diff --git a/Engine/Plugins/Experimental/NeuralNetworkInference/Source/NeuralNetworkInference/Public/NeuralNetwork.h b/Engine/Plugins/Experimental/NeuralNetworkInference/Source/NeuralNetworkInference/Public/NeuralNetwork.h
--- a/Engine/Plugins/Experimental/NeuralNetworkInference/Source/NeuralNetworkInference/Public/NeuralNetwork.h
+++ b/Engine/Plugins/Experimental/NeuralNetworkInference/Source/NeuralNetworkInference/Public/NeuralNetwork.h
@@ -4,10 +4,15 @@
 
 #include "CoreMinimal.h"
 #include "NeuralEnumClasses.h"
+#include "NeuralTensor.h"
 #include "NeuralTensors.h"
+
 #include "NeuralNetwork.generated.h"
 
+class FNeuralOperator;
 class UAssetImportData;
+struct FModelProto;
+struct FNeuralTensorManager;
 
 /**
  * Whether UNeuralNetwork::Run() will block the thread until completed (Synchronous), or whether it will run on a background thread, not blocking the calling thread (Asynchronous).
@@ -77,22 +82,22 @@ public:
 
 	/**
 	 * Editor-only function
-	 * It loads the desired UNeuralNetwork definition and its weights from an ONNX/ORT file.
+	 * It loads the desired network graph definition and weights given an input ONNX file path.
 	 * @param InModelFilePath can either be a full path or a relative path with respect to the Game project.
-	 * @return Whether it loaded the UNeuralNetwork successfully.
+	 * @return Whether the network was successfully loaded.
 	 */
 #if WITH_EDITOR
 	bool Load(const FString& InModelFilePath);
 #endif // WITH_EDITOR
 
 	/**
-	 * It loads the UNeuralNetwork definition and its weights into ORT.
-	 * @return Whether it loaded the UNeuralNetwork successfully.
+	 * It loads the desired network graph definition and weights given an input UNeuralNetwork UAsset.
+	 * @return Whether the network was successfully loaded.
 	 */
 	bool Load();
 
 	/**
-	 * It returns whether it loaded the UNeuralNetwork successfully.
+	 * It returns whether a network is currently loaded.
 	 */
 	bool IsLoaded() const;
 
@@ -121,6 +126,16 @@ public:
 	const FNeuralTensor& GetOutputTensor(const int32 InTensorIndex = 0) const;
 	const FNeuralTensors& GetOutputTensors() const;
 
+	/**
+	 * Callback when UNeuralNetwork::Run() is finished. It will only occur if SynchronousMode == ENeuralNetworkSynchronousMode::Asynchronous and it could be triggered from any thread.
+	 */
+	DECLARE_DELEGATE(FOnAsyncRunCompleted);
+
+	/**
+	 * Returns the delegate called when async Run() is over (on any thread)
+	 */
+	FOnAsyncRunCompleted& GetOnAsyncRunCompletedDelegate();
+
 	/**
 	 * Run() executes the forward pass on the current UNeuralNetwork given the current input FDeprecatedNeuralTensor(s), which were previously filled with
 	 * SetInputFromArrayCopy() or GetInputDataPointerMutable().
@@ -130,6 +145,9 @@ public:
 	void Run();
 
 protected:
+	UPROPERTY(VisibleAnywhere, Category = "Neural Network Inference")
+	bool bIsLoaded;
+
 	/**
 	 * Whether Run() will use CPU or GPU acceleration hardware.
 	 * If SetDeviceType() is never called, the default device (EDeviceType::CPU) will be used.
@@ -185,7 +203,8 @@ private:
 	UPROPERTY()
 	TArray<uint8> ModelReadFromDiskInBytes;
 
-	bool bIsLoaded;
+	FOnAsyncRunCompleted OnAsyncRunCompletedDelegate;
+
 	/**
 	 * PIMPL idiom to hide 3rd party dependencies
 	 * http://www.cppsamples.com/common-tasks/pimpl.html
