NeuralNetworkInference (NNI) plugin:
- (F)NeuralNetworkFromONNXTranslator renamed as (F)GraphProtoToNeuralNetworkConverter to keep folder cleaner.

#rnx

diff --git a/Engine/Plugins/Experimental/NeuralNetworkInference/Source/NeuralNetworkInference/Private/GraphProtoToNeuralNetworkConverter.cpp b/Engine/Plugins/Experimental/NeuralNetworkInference/Source/NeuralNetworkInference/Private/GraphProtoToNeuralNetworkConverter.cpp
new file mode 100644
--- /dev/null
+++ b/Engine/Plugins/Experimental/NeuralNetworkInference/Source/NeuralNetworkInference/Private/GraphProtoToNeuralNetworkConverter.cpp
@@ -0,0 +1,614 @@
+// Copyright Epic Games, Inc. All Rights Reserved.
+
+#include "GraphProtoToNeuralNetworkConverter.h"
+#include "Misc/Paths.h"
+#include "NeuralNetworkInferenceUtils.h"
+#include "NeuralOperator.h"
+#include "NeuralOperators.h"
+#include "NeuralTensorManager.h"
+
+#if WITH_EDITOR
+#if PLATFORM_WINDOWS
+#include "ModelProtoFileReader.h"
+#endif //PLATFORM_WINDOWS
+#endif //WITH_EDITOR
+
+
+
+/* FGraphProtoToNeuralNetworkConverter public functions
+ *****************************************************************************/
+
+bool FGraphProtoToNeuralNetworkConverter::Translate(TArray<TSharedPtr<FNeuralOperator>>& OutOperators, FNeuralTensorManager& InTensorManager, const FGraphProto& InGraphProto)
+{
+	TMap<FString, int32> NameIndexMap = InTensorManager.GetNameIndexMap();
+	TMap<FString, int32> DummyOutputNameIndexMap;
+	const bool bResult = CreateOperatorsAndEditTensorArray(OutOperators, DummyOutputNameIndexMap, InTensorManager.GetTensorsMutable(), NameIndexMap, InTensorManager.GetInputNameIndexMap(),
+		InTensorManager.GetOutputNameIndexMap(), InGraphProto);
+	if (NameIndexMap.Num() != InTensorManager.GetNameIndexMap().Num())
+	{
+		UE_LOG(LogNeuralNetworkInference, Warning, TEXT("FGraphProtoToNeuralNetworkConverter::Translate(): NameIndexMap changed its size, it must not (%d != %d)."),
+			NameIndexMap.Num(), InTensorManager.GetNameIndexMap().Num());
+	}
+	return bResult;
+}
+
+#if WITH_EDITOR
+bool FGraphProtoToNeuralNetworkConverter::Translate(TArray<TSharedPtr<FNeuralOperator>>& OutOperators, FNeuralTensorManager& OutTensorManager, const FGraphProto& InGraphProto,
+	const FString& InModelPath)
+{
+	// MaxNumberTensorsInNetwork is required to avoid the bug of resizing and invalidating all pointers given to the operators
+	int32 MaxNumberTensorsInNetwork = 0;
+	{
+		// Quickly estimate the number of FNeuralTensor's of the UNeuralNetworkLegacy for Tensors.Reserve(...)
+		// Number FNeuralTensor's = Number (unique) input FNeuralTensors + number (unique) output FNeuralTensors
+		// Uniqueness guaranteed by using the TSet
+		TSet<FString> TensorNameSet;
+		for (const FNodeProto& NodeProto : InGraphProto.Node)
+		{
+			// Check input tensors
+			for (const FString& TensorName : NodeProto.Input)
+			{
+				if (TensorNameSet.Find(TensorName) == nullptr)
+				{
+					TensorNameSet.Add(TensorName);
+				}
+			}
+			// Check output tensors
+			for (const FString& TensorName : NodeProto.Output)
+			{
+				if (TensorNameSet.Find(TensorName) == nullptr)
+				{
+					TensorNameSet.Add(TensorName);
+				}
+			}
+		}
+		MaxNumberTensorsInNetwork = TensorNameSet.Num()
+		// Add maximum number possible of auxiliary tensors. A conservative estimate is the number of operators times the maximum possible number of auxiliary tensors per operator
+			+ InGraphProto.Node.Num() * FNeuralOperator::GetMaximumPossibleNumberAuxiliaryTensorsPerOperator();
+	}
+	TArray<FNeuralTensor> Tensors;
+	Tensors.Reserve(MaxNumberTensorsInNetwork);
+	TMap<FString, int32> NameIndexMap;
+
+	// Read input tensors
+	TMap<FString, int32> InputNameIndexMap;
+	for (const FValueInfoProto& GraphProtoInput : InGraphProto.Input)
+	{
+		// Create tensor
+		const TArray<FTensorShapeProtoDimension>& Dimensions = GraphProtoInput.Type.TensorType.Shape.Dim;
+		TArray<int64> DimensionsAsTArray;
+		for (const FTensorShapeProtoDimension& Dimension : Dimensions)
+		{
+			DimensionsAsTArray.Emplace(Dimension.DimValue);
+		}
+		// Add operator name to TMap
+		NameIndexMap.Add(GraphProtoInput.Name, Tensors.Num());
+		InputNameIndexMap.Add(GraphProtoInput.Name, Tensors.Num());
+		// Add tensor to TMap
+		Tensors.Emplace(FNeuralTensor(ENeuralDataType::Float, DimensionsAsTArray, GraphProtoInput.Name, ENeuralTensorTypeGPU::Input));
+	}
+
+	// Read output tensor names (but do not create them yet, it has to happen simultaneously with the operator creation and tensor inlining)
+	// Issue: For inlining to work properly, output tensors cannot be created until after inlining operators (i.e., CreateOperatorsAndEditTensorArray).
+	// Otherwise, the following failing case could happen with a network like: InputTensor --> Gemm --> IntermediateTensor --> Relu --> OutputTensor
+	// 1. Inputs and outputs are created before hand, i.e., InputTensor and OutputTensor
+	// 2. CreateOperatorsAndEditTensorArray() is run
+	//   2.1. Gemm and IntermediateTensor would be created (because Gemm cannot be inlined)
+	//   2.2. Relu would be created, Relu would verify its input can be inlined, but both IntermediateTensor and OutputTensor are already created at this point, so it cannot inline them (crash/bug)
+	TMap<FString, int32> OutputNameDummyIndexMap;
+	for (const FValueInfoProto& GraphProtoOutput : InGraphProto.Output)
+	{
+		// Add operator name to TMap
+		if (NameIndexMap.Find(GraphProtoOutput.Name))
+		{
+			UE_LOG(LogNeuralNetworkInference, Warning, TEXT("FGraphProtoToNeuralNetworkConverter::Translate(): A FNeuralTensor of the UNeuralNetworkLegacy is a global Input and Output simultaneously."));
+		}
+		OutputNameDummyIndexMap.Add(GraphProtoOutput.Name, -1);
+	}
+
+	// Fill operators
+	TMap<FString, int32> OutputNameIndexMap;
+	bool bWasSuccessful = CreateOperatorsAndEditTensorArray(OutOperators, OutputNameIndexMap, Tensors, NameIndexMap, InputNameIndexMap, OutputNameDummyIndexMap, InGraphProto, InModelPath);
+	if (Tensors.Num() <= NameIndexMap.Num())
+	{
+		UE_LOG(LogNeuralNetworkInference, Display, TEXT("Auto-inlining: %d initial tensors compressed into %d (%d tensors merged)."),
+			NameIndexMap.Num(), Tensors.Num(), NameIndexMap.Num() - Tensors.Num());
+	}
+	// Sanity check
+	if (Tensors.Num() > MaxNumberTensorsInNetwork)
+	{
+		UE_LOG(LogNeuralNetworkInference, Warning, TEXT("FGraphProtoToNeuralNetworkConverter::Translate(): Tensors.Num() <= MaxNumberTensorsInNetwork failed (%d vs. %d). A resize of Tensors"
+			" might have occurred, invalidating all TArray<FNeuralTensor*> FOperator variables. Report this error (so we can fix MaxNumberTensorsInNetwork to be big enough)."),
+			Tensors.Num(), MaxNumberTensorsInNetwork);
+		bWasSuccessful = false;
+	}
+
+	// Set OutTensorManager
+	if (bWasSuccessful)
+	{
+		return OutTensorManager.Load(Tensors, NameIndexMap, InputNameIndexMap, OutputNameIndexMap);
+	}
+	return bWasSuccessful;
+}
+#endif //WITH_EDITOR
+
+
+
+/* FGraphProtoToNeuralNetworkConverter private functions
+ *****************************************************************************/
+
+bool FGraphProtoToNeuralNetworkConverter::CreateOperatorsAndEditTensorArray(TArray<TSharedPtr<FNeuralOperator>>& OutOperators, TMap<FString, int32>& OutputNameIndexMap,
+	TArray<FNeuralTensor>& InOutTensors, TMap<FString, int32>& InOutNameIndexMap, const TMap<FString, int32>& InInputNameDummyIndexMap,
+	const TMap<FString, int32>& InOutputNameDummyIndexMap, const FGraphProto& InGraphProto, const FString& InModelPath)
+{
+	const bool bConvertingFromSerializedData = (InModelPath.Len() == 0);
+	const int64 InOutTensorsNumInit = InOutTensors.Num();
+	const int64 InOutNameIndexMapNumInit = InOutNameIndexMap.Num();
+	// Create operators sequentially
+	OutOperators.Empty();
+	OutputNameIndexMap.Empty();
+	// Iterating over each operator
+	for (const FNodeProto& NodeProto : InGraphProto.Node)
+	{
+		// Read/set input tensors
+		TArray<FNeuralTensor*> OperatorInputTensors;
+		TArray<FString> OperatorInputTensorNames;
+		// Iterating over each input tensor name for that operator
+		for (const FString& InputTensorName : NodeProto.Input)
+		{
+			OperatorInputTensorNames.Push(InputTensorName);
+			// FNeuralTensor already read (e.g., absolute input of the UNeuralNetworkLegacy, or the output of a previous tensor required as input of this one)
+			if (int32* ExistingTensorIndex = InOutNameIndexMap.Find(InputTensorName))
+			{
+				OperatorInputTensors.Push(&InOutTensors[*ExistingTensorIndex]);
+			}
+			// FNeuralTensor was never modified, so it must come from a weight tensor (regardless of whether it is also an absolute output of the network)
+			else
+			{
+#if WITH_EDITOR
+				const FTensorProto& TensorProto = *FModelProto::FindElementInArray(InputTensorName, InGraphProto.Initializer, /*bMustValueBeFound*/true);
+				// Is it also an absolute output of the network?
+				const int64 TensorIndex = InOutTensors.Num();
+				if (!bConvertingFromSerializedData && InOutputNameDummyIndexMap.Find(InputTensorName))
+				{
+					OutputNameIndexMap.Add(InputTensorName, TensorIndex);
+				}
+				// Create and add tensor to InOutTensors
+				InOutTensors.Push(FNeuralTensor());
+				FNeuralTensor& NewTensor = InOutTensors.Last();
+				// Deprecated code for otxt files, if added back, add before the 2 for loops for speedup
+				// FString PathPart, FilenamePart, ExtensionPart;
+				// if (!bConvertingFromSerializedData)
+				// {
+				// 	FPaths::Split(InModelPath, PathPart, FilenamePart, ExtensionPart);
+				// }
+				// Get NewTensor from TensorProto
+				if (!NewTensor.SetFromTensorProto(&TensorProto, InputTensorName, ENeuralTensorTypeGPU::Weight/*, PathPart*/))
+				{
+					UE_LOG(LogNeuralNetworkInference, Warning, TEXT("FGraphProtoToNeuralNetworkConverter::CreateOperatorsAndEditTensorArray(): SetFromTensorProto() failed."));
+					return false;
+				}
+				// Update InOutNameIndexMap and OperatorInputTensors
+				InOutNameIndexMap.Add(InputTensorName, TensorIndex);
+				OperatorInputTensors.Push(&NewTensor);
+#else
+				UE_LOG(LogNeuralNetworkInference, Warning, TEXT("FGraphProtoToNeuralNetworkConverter::CreateOperatorsAndEditTensorArray(): This code should never be called in non-Editor mode."));
+				return false;
+#endif //WITH_EDITOR
+			}
+		}
+		// Nodes that can be replaced for inlined functions and do not need to be executed
+		if (NodeProto.OperatorType == TEXT("Constant"))
+		{
+			// Sanity check
+			if (NodeProto.Output.Num() != 1)
+			{
+				UE_LOG(LogNeuralNetworkInference, Warning, TEXT("FGraphProtoToNeuralNetworkConverter::CreateOperatorsAndEditTensorArray(): Constant operator requires exactly 1 output, not %d."), NodeProto.Output.Num());
+				return false;
+			}
+			const FString& OutputTensorName = NodeProto.Output[0];
+			// Case 1: bConvertingFromSerializedData
+			if (bConvertingFromSerializedData)
+			{
+				// Sanity check
+				if (!InOutNameIndexMap.Find(OutputTensorName))
+				{
+					UE_LOG(LogNeuralNetworkInference, Warning, TEXT("FGraphProtoToNeuralNetworkConverter::CreateOperatorsAndEditTensorArray()-Constant: Tensor %s was not found."), *OutputTensorName);
+					return false;
+				}
+				// Do nothing else
+			}
+			// Case 2: !bConvertingFromSerializedData
+			else
+			{
+				// Sanity check
+				if (InOutNameIndexMap.Find(OutputTensorName))
+				{
+					UE_LOG(LogNeuralNetworkInference, Warning, TEXT("FGraphProtoToNeuralNetworkConverter::CreateOperatorsAndEditTensorArray()-Constant: Tensor %s was already defined and it should have not."), *OutputTensorName);
+					return false;
+				}
+				// Is it also an absolute output of the network?
+				const int64 TensorIndex = InOutTensors.Num();
+				if (InOutputNameDummyIndexMap.Find(OutputTensorName))
+				{
+					OutputNameIndexMap.Add(OutputTensorName, TensorIndex);
+				}
+				// Create tensor
+				InOutNameIndexMap.Add(OutputTensorName, TensorIndex);
+				InOutTensors.Push(FNeuralTensor(ENeuralDataType::None, TArray<int64>({}), OutputTensorName, ENeuralTensorTypeGPU::Weight));
+
+				// Create constant operator layer and get final tensor value
+				FConstantOperator ConstantOperator(NodeProto);
+				ConstantOperator.SetOutputTensors({&InOutTensors.Last()});
+				// Sanity checks
+				if (!ConstantOperator.ConfigureOutputAndInternalVariablesAndSanityChecks())
+				{
+					UE_LOG(LogNeuralNetworkInference, Warning, TEXT("FGraphProtoToNeuralNetworkConverter::CreateOperatorsAndEditTensorArray()-Constant: Tensor %s was already defined and it should have not."), *OutputTensorName);
+					return false;
+				}
+				// Get output
+				ConstantOperator.ForwardCPU();
+			}
+		}
+		// Nodes/Operators that must be created and used on each Run() pass of the network
+		else
+		{
+			// Create and configure operator
+			TSharedPtr<FNeuralOperator> Operator = CreateOperator(NodeProto, OperatorInputTensorNames, InInputNameDummyIndexMap, InOutputNameDummyIndexMap, InGraphProto);
+			// ConvTranspose requires flipping OperatorInputTensors[1], so we expect for now a Weight type, so it can be flipped beforehand
+			if (!bConvertingFromSerializedData && NodeProto.OperatorType == TEXT("ConvTranspose")) //if (NodeProto.OperatorType.Contains(TEXT("Conv"), ESearchCase::CaseSensitive))
+			{
+				UE_LOG(LogNeuralNetworkInference, Display, TEXT("Flipping FConvTranspose weights. This is a temporary measurement but it will keep the weights flipped for now."));
+				// Sanity check
+				if (OperatorInputTensors.Num() > 1 && OperatorInputTensors[1]->GetTensorTypeGPU() != ENeuralTensorTypeGPU::Weight)
+				{
+					UE_LOG(LogNeuralNetworkInference, Warning,
+						TEXT("FGraphProtoToNeuralNetworkConverter::CreateOperatorsAndEditTensorArray(): For now, %s operators require a constant InputTensor[0] (i.e., a fixed weight tensor), it cannot change."), *NodeProto.OperatorType);
+					return false;
+				}
+				// Flip (invert) OperatorInputTensors[1]
+				if (!OperatorInputTensors[1]->Flip(2, OperatorInputTensors[1]->GetNumberDimensions()))
+				{
+					UE_LOG(LogNeuralNetworkInference, Warning, TEXT("FGraphProtoToNeuralNetworkConverter::CreateOperatorsAndEditTensorArray(): Tensor::Flip() failed."));
+					return false;
+				}
+			}
+			// Set input tensors in Operator
+			const int32 InlinedTensorIndex = Operator->SetInputTensorsAndGetInlinedIndex(OperatorInputTensors);
+			// Read/set output tensors
+			// 2 different cases for OperatorOutputTensors:
+				// Case 1/2: bConvertingFromSerializedData --> All FNeuralTensors already exist
+					// Subcase 1/2: InlinedTensorIndex > -1 --> Do nothing
+					// Subcase 2/2: InlinedTensorIndex < 0 --> Fill OperatorOutputTensors from known index
+				// Case 2/2: !bConvertingFromSerializedData
+					// Subcase 1/3: ExistingTensorIndex is found --> FNeuralTensor already exists
+					// Subcase 2/3: InlinedTensorIndex < 0 --> Output tensor must be created because operator creates its own output
+					// Subcase 3/3: InlinedTensorIndex > -1 --> Output tensor is the same than 1 of the input tensors of that operator (inlined operator)
+			if (InlinedTensorIndex > -1 && NodeProto.Output.Num() != 1)
+			{
+				UE_LOG(LogNeuralNetworkInference, Warning, TEXT("FGraphProtoToNeuralNetworkConverter::CreateOperatorsAndEditTensorArray(): There cannot be an inlined tensor (index %d)"
+					" if NodeProto.Output.Num() > 1 (currently %d)."), InlinedTensorIndex, NodeProto.Output.Num());
+				return false;
+			}
+			TArray<FNeuralTensor*> OperatorOutputTensors;
+			for (const FString& OutputTensorName : NodeProto.Output)
+			{
+				// Case 1: bConvertingFromSerializedData
+				if (bConvertingFromSerializedData)
+				{
+					// Subcase 1/2: Do nothing
+					// Subcase 2/2: Fill OperatorOutputTensors from known index
+					if (InlinedTensorIndex < 0)
+					{
+						const int32 ExistingTensorIndex = InOutNameIndexMap.FindChecked(OutputTensorName);
+						OperatorOutputTensors.Push(&InOutTensors[ExistingTensorIndex]);
+					}
+				}
+				// Case 2: !bConvertingFromSerializedData
+				else
+				{
+					// Subcase 1/3: FNeuralTensor already exists (i.e., absolute output of the UNeuralNetworkLegacy or bConvertingFromSerializedData == true)
+					if (const int32* const ExistingTensorIndex = InOutNameIndexMap.Find(OutputTensorName))
+					{
+						if (InlinedTensorIndex > -1)
+						{
+							UE_LOG(LogNeuralNetworkInference, Warning,
+								TEXT("FGraphProtoToNeuralNetworkConverter::CreateOperatorsAndEditTensorArray(): InlinedTensorIndex = %d for this operator (%s) but the tensor (%s) already exists."), InlinedTensorIndex, *NodeProto.OperatorType, *OutputTensorName);
+							return false;
+						}
+						OperatorOutputTensors.Push(&InOutTensors[*ExistingTensorIndex]);
+						// Is it also an absolute output of the network?
+						if (InOutputNameDummyIndexMap.Find(OutputTensorName))
+						{
+							OutputNameIndexMap.Add(OutputTensorName, *ExistingTensorIndex);
+							InOutTensors[*ExistingTensorIndex].SetTensorTypeGPU(ENeuralTensorTypeGPU::Output);
+						}
+					}
+					// Subcases 2 and 3: FNeuralTensor not found, it must be defined now (i.e., intermediate tensor of the UNeuralNetworkLegacy)
+					// Subcase 2/3: Output tensor must be created because operator creates its own output
+					else if (InlinedTensorIndex < 0)
+					{
+						// Is it also an absolute output of the network?
+						const int64 TensorIndex = InOutTensors.Num();
+						ENeuralTensorTypeGPU TensorTypeGPU = ENeuralTensorTypeGPU::IntermediateNotInitialized;
+						if (InOutputNameDummyIndexMap.Find(OutputTensorName))
+						{
+							OutputNameIndexMap.Add(OutputTensorName, TensorIndex);
+							TensorTypeGPU = ENeuralTensorTypeGPU::Output;
+						}
+						// Create tensor
+						InOutNameIndexMap.Add(OutputTensorName, TensorIndex);
+						InOutTensors.Push(FNeuralTensor(ENeuralDataType::Float, TArray<int64>({}), OutputTensorName, TensorTypeGPU));
+						OperatorOutputTensors.Push(&InOutTensors.Last());
+					}
+					// Subcase 3/3: Output tensor is the same than 1 of the input tensors of that operator (inlined operator)
+					else
+					{
+						if (NodeProto.Output.Num() != 1)
+						{
+							UE_LOG(LogNeuralNetworkInference, Warning, TEXT("FGraphProtoToNeuralNetworkConverter::CreateOperatorsAndEditTensorArray()-%s-%s:"
+								" NodeProto.Output.Num() = %d != 1. Only implemented for inlined operators with 1 output. Potential issues for operators like PReLU."
+								" Although it should work as long as we pre-define the desired order of FNeuralTensors inside that operator."),
+								*NodeProto.OperatorType, *Operator->GetName(), OperatorInputTensors.Num(), NodeProto.Output.Num());
+							return false;
+						}
+						const FString& NodeProtoInputName = NodeProto.Input[InlinedTensorIndex];
+						const int32 InlinedTensorGlobalIndex = InOutNameIndexMap.FindChecked(NodeProtoInputName);
+						InOutNameIndexMap.Add(OutputTensorName, InlinedTensorGlobalIndex);
+						// Is it also an absolute output of the network?
+						if (InOutputNameDummyIndexMap.Find(OutputTensorName))
+						{
+							OutputNameIndexMap.Add(OutputTensorName, InlinedTensorGlobalIndex);
+							InOutTensors[InlinedTensorGlobalIndex].SetTensorTypeGPU(ENeuralTensorTypeGPU::Output);
+						}
+					}
+				}
+			}
+			// Set output tensors in Operator (if operator is not fully inlined)
+			if (OperatorOutputTensors.Num() > 0)
+			{
+				Operator->SetOutputTensors(OperatorOutputTensors);
+			}
+			// Create and add auxiliary tensors for that operator (if any)
+			const int64 NumberAuxiliaryTensors = Operator->GetNumberAuxiliaryTensors();
+			if (NumberAuxiliaryTensors > 0)
+			{
+				TArray<FNeuralTensor*> AuxiliaryOperatorTensors;
+				for (int64 AuxiliaryTensorIndex = 0; AuxiliaryTensorIndex < Operator->GetNumberAuxiliaryTensors(); ++AuxiliaryTensorIndex)
+				{
+					const FString AuxiliaryTensorName = NodeProto.Name + TEXT("-") + Operator->GetName() + TEXT("-Aux") + FString::FromInt(AuxiliaryTensorIndex);
+					// Uasset file
+					if (bConvertingFromSerializedData)
+					{
+						// Find and push auxiliary tensor
+						if (const int32* const AuxiliaryTensorAbsoluteIndex = InOutNameIndexMap.Find(AuxiliaryTensorName))
+						{
+							AuxiliaryOperatorTensors.Push(&InOutTensors[*AuxiliaryTensorAbsoluteIndex]);
+						}
+						// Something went wrong, the tensor should exist
+						else
+						{
+							UE_LOG(LogNeuralNetworkInference, Warning, TEXT("FGraphProtoToNeuralNetworkConverter::CreateOperatorsAndEditTensorArray(): AuxiliaryTensorName %s did not exist in InOutTensors."), *AuxiliaryTensorName);
+							return false;
+						}
+					}
+					// ONNX file
+					else
+					{
+						// Sanity check
+						if (InOutNameIndexMap.Find(AuxiliaryTensorName))
+						{
+							UE_LOG(LogNeuralNetworkInference, Warning, TEXT("FGraphProtoToNeuralNetworkConverter::CreateOperatorsAndEditTensorArray(): InOutNameIndexMap already has %s. Notify us of this error."), *AuxiliaryTensorName);
+							return false;
+						}
+						// Create and push new auxiliary tensor
+						InOutNameIndexMap.Add(AuxiliaryTensorName, InOutTensors.Num());
+						InOutTensors.Push(FNeuralTensor(ENeuralDataType::None, TArray<int64>({}), AuxiliaryTensorName, ENeuralTensorTypeGPU::Generic)); // ENeuralTensorTypeGPU can be modify by FNeuralOperator::Configure()
+						AuxiliaryOperatorTensors.Push(&InOutTensors.Last());
+					}
+				}
+				Operator->SetAuxiliaryTensors(AuxiliaryOperatorTensors);
+			}
+			// If operator is inlined, then input and output should already match
+			// Configure operator
+			if (!Operator->Configure())
+			{
+				UE_LOG(LogNeuralNetworkInference, Warning,
+					TEXT("FGraphProtoToNeuralNetworkConverter::CreateOperatorsAndEditTensorArray(): Operator %s could not be configured, so network could not be loaded."), *Operator->GetName());
+				return false;
+			}
+			// Push Operator into OutOperators
+			OutOperators.Push(Operator);
+		}
+	}
+	// Sanity checks for bConvertingFromSerializedData
+	if (bConvertingFromSerializedData)
+	{
+		if (OutputNameIndexMap.Num() > 0)
+		{
+			UE_LOG(LogNeuralNetworkInference, Warning,
+				TEXT("FGraphProtoToNeuralNetworkConverter::CreateOperatorsAndEditTensorArray(): Error in this function, OutputNameIndexMap.Num() should be 0 and not %d when reading from serialized data."),
+				OutputNameIndexMap.Num());
+			return false;
+		}
+		else if (InOutTensorsNumInit != InOutTensors.Num())
+		{
+			UE_LOG(LogNeuralNetworkInference, Warning,
+				TEXT("FGraphProtoToNeuralNetworkConverter::CreateOperatorsAndEditTensorArray(): Error in this function, InOutTensors.Num() should not change from %d to %d when reading from serialized data."),
+				InOutTensorsNumInit, InOutTensors.Num());
+			return false;
+		}
+		else if (InOutNameIndexMapNumInit != InOutNameIndexMap.Num())
+		{
+			UE_LOG(LogNeuralNetworkInference, Warning,
+				TEXT("FGraphProtoToNeuralNetworkConverter::CreateOperatorsAndEditTensorArray(): Error in this function, InOutNameIndexMap.Num() should not change from %d to %d when reading from serialized data."),
+				InOutTensorsNumInit, InOutTensors.Num());
+			return false;
+		}
+	}
+	// Sanity check for !bConvertingFromSerializedData
+	else if (OutputNameIndexMap.Num() != InOutputNameDummyIndexMap.Num())
+	{
+		UE_LOG(LogNeuralNetworkInference, Warning,
+			TEXT("FGraphProtoToNeuralNetworkConverter::CreateOperatorsAndEditTensorArray(): OutputNameIndexMap.Num() == InOutputNameDummyIndexMap.Num() failed, %d != %d."
+				" This means that not all of the outputs defined on the ONNX file were actually needed. The program will proceed with only the required ones."),
+			OutputNameIndexMap.Num(), InOutputNameDummyIndexMap.Num());
+	}
+	// Sanity check for all cases
+	if (InOutNameIndexMap.Num() < InOutTensors.Num())
+	{
+		UE_LOG(LogNeuralNetworkInference, Warning,
+			TEXT("FGraphProtoToNeuralNetworkConverter::CreateOperatorsAndEditTensorArray(): Error in this function, InOutNameIndexMap.Num() >= InOutTensors.Num() failed, %d != %d. They"
+				" should match if no inlining occurs, or more generically, InOutNameIndexMap.Num() = InOutTensors.Num() + number_inlined_tensors."),
+			InOutNameIndexMap.Num(), InOutTensors.Num());
+	}
+
+	return (OutOperators.Num() > 0);
+}
+
+bool FGraphProtoToNeuralNetworkConverter::CanTensorBeInlined(const FString& InTensorName, const TMap<FString, int32>& InInputNameDummyIndexMap,
+	const TMap<FString, int32>& InOutputNameDummyIndexMap, const FNodeProto& InNodeProto, const FGraphProto& InGraphProto, const bool bCanOutputPropertiesChangeWithPostForward)
+{
+	// If tensor is one of the input/output tensors --> Operator cannot be inlined
+	if (InInputNameDummyIndexMap.Find(InTensorName) || InOutputNameDummyIndexMap.Find(InTensorName))
+	{
+		return false;
+	}
+	// If PostForward() changes the properties of the output tensor (e.g., Reshape might change its size), then it cannot be combined with a global output of the network
+	else if (bCanOutputPropertiesChangeWithPostForward)
+	{
+		for (const FValueInfoProto& GraphProtoOutput : InGraphProto.Output)
+		{
+			if (GraphProtoOutput.Name == InNodeProto.Output[0])
+			{
+				return false;
+			}
+		}
+	}
+	// If tensor is used again --> Operator cannot be inlined
+	else
+	{
+		int32 NumberOperatorsUsingTensor = 0;
+		// Iterating over each operator
+		for (const FNodeProto& NodeProto : InGraphProto.Node)
+		{
+			// Iterating over each input tensor name for that operator
+			for (const FString& InputTensorName : NodeProto.Input)
+			{
+				// If names matches, increase NumberOperatorsUsingTensor
+				if (InputTensorName == InTensorName)
+				{
+					++NumberOperatorsUsingTensor;
+					if (NumberOperatorsUsingTensor > 1)
+					{
+						return false;
+					}
+				}
+			}
+		}
+		// Sanity check
+		if (NumberOperatorsUsingTensor != 1)
+		{
+			UE_LOG(LogNeuralNetworkInference, Warning,
+				TEXT("FGraphProtoToNeuralNetworkConverter::CanTensorBeInlined(): At this point, NumberOperatorsUsingTensor should be exactly 1 (not %d)."),
+				NumberOperatorsUsingTensor);
+		}
+	}
+	// None of the above is true --> Operator can be safely inlined
+	return true;
+}
+
+TSet<uint32> FGraphProtoToNeuralNetworkConverter::GetPotentiallyInlinedTensors(const TArray<FString>& InTensorNames, const TMap<FString, int32>& InInputNameDummyIndexMap,
+	const TMap<FString, int32>& InOutputNameDummyIndexMap, const FNodeProto& InNodeProto, const FGraphProto& InGraphProto)
+{
+	TSet<uint32> PotentiallyInlinedTensors;
+	for (int32 TensorNameIndex = 0; TensorNameIndex < InTensorNames.Num(); ++TensorNameIndex)
+	{
+		const FString& TensorName = InTensorNames[TensorNameIndex];
+		if (FGraphProtoToNeuralNetworkConverter::CanTensorBeInlined(TensorName, InInputNameDummyIndexMap, InOutputNameDummyIndexMap, InNodeProto, InGraphProto))
+		{
+			PotentiallyInlinedTensors.Add((uint32)TensorNameIndex);
+		}
+	}
+	return PotentiallyInlinedTensors;
+}
+
+// Eg Element-wise operators with 0 attributes
+#define IF_MATCHED_INLINABLE_OPERATOR_INITIALIZATION(OperatorTypeSuffix) \
+	if (InNodeProto.OperatorType == TEXT(#OperatorTypeSuffix)) /* Eg TEXT("Relu") */ \
+	{ \
+		const bool bCanTensorBeInlined = CanTensorBeInlined(InInputTensorNamesForOperator[0], InInputNameDummyIndexMap, InOutputNameDummyIndexMap, InNodeProto, InGraphProto); \
+		Operator = MakeShared<F##OperatorTypeSuffix##Operator>(bCanTensorBeInlined); /* Eg FReluOperator */ \
+	}
+
+// Eg Element-wise operators with n attributes or operators like BatchNormalization
+#define IF_MATCHED_INLINABLE_OPERATOR_WITH_PROTO_INITIALIZATION(OperatorTypeSuffix, bCanOutputPropertiesChangeWithPostForward) \
+	if (InNodeProto.OperatorType == TEXT(#OperatorTypeSuffix)) /* Eg TEXT("BatchNormalization") */ \
+	{ \
+		const bool bCanTensorBeInlined = CanTensorBeInlined(InInputTensorNamesForOperator[0], InInputNameDummyIndexMap, InOutputNameDummyIndexMap, InNodeProto, InGraphProto, bCanOutputPropertiesChangeWithPostForward); \
+		Operator = MakeShared<F##OperatorTypeSuffix##Operator>(bCanTensorBeInlined, InNodeProto); /* Eg FBatchNormalizationOperator */ \
+	}
+
+// Eg Conv(Transpose) and Gemm
+#define IF_MATCHED_OPERATOR_WITH_PROTO_INITIALIZATION(OperatorTypeSuffix) \
+	if (InNodeProto.OperatorType == TEXT(#OperatorTypeSuffix)) /* Eg TEXT("Conv") */ \
+	{ \
+		Operator = MakeShared<F##OperatorTypeSuffix##Operator>(InNodeProto); /* Eg FConvOperator */ \
+	}
+
+#define IF_MATCHED_MULTIDIRECTIONAL_BROADCASTING_OPERATOR_INITIALIZATION(OperatorTypeSuffix) \
+	if (InNodeProto.OperatorType == TEXT(#OperatorTypeSuffix)) /* Eg TEXT("Add") */ \
+	{ \
+		const TSet<uint32> PotentiallyInlinedTensors = GetPotentiallyInlinedTensors(InInputTensorNamesForOperator, InInputNameDummyIndexMap, InOutputNameDummyIndexMap, InNodeProto, InGraphProto); \
+		Operator = MakeShared<F##OperatorTypeSuffix##Operator>(PotentiallyInlinedTensors); /* Eg FAddOperator */ \
+	}
+
+TSharedPtr<FNeuralOperator> FGraphProtoToNeuralNetworkConverter::CreateOperator(const FNodeProto& InNodeProto, const TArray<FString>& InInputTensorNamesForOperator,
+	const TMap<FString, int32>& InInputNameDummyIndexMap, const TMap<FString, int32>& InOutputNameDummyIndexMap, const FGraphProto& InGraphProto)
+{
+	TSharedPtr<FNeuralOperator> Operator;
+	IF_MATCHED_INLINABLE_OPERATOR_INITIALIZATION(Abs)
+	else IF_MATCHED_INLINABLE_OPERATOR_INITIALIZATION(Acos)
+	else IF_MATCHED_MULTIDIRECTIONAL_BROADCASTING_OPERATOR_INITIALIZATION(Add)
+	else IF_MATCHED_INLINABLE_OPERATOR_INITIALIZATION(Asin)
+	else IF_MATCHED_INLINABLE_OPERATOR_INITIALIZATION(Atan)
+	else IF_MATCHED_INLINABLE_OPERATOR_WITH_PROTO_INITIALIZATION(BatchNormalization, false)
+	else IF_MATCHED_INLINABLE_OPERATOR_INITIALIZATION(Ceil)
+	else IF_MATCHED_OPERATOR_WITH_PROTO_INITIALIZATION(Conv)
+	else IF_MATCHED_OPERATOR_WITH_PROTO_INITIALIZATION(ConvTranspose)
+	if (InNodeProto.OperatorType == TEXT("ConvTranspose"))
+	{
+		TSharedPtr<FConvTransposeOperator> ConvTransposeOperator;
+		ConvTransposeOperator = MakeShared<FConvTransposeOperator>(InNodeProto);
+		ConvTransposeOperator->SetWhetherWeightTensorWasFlipped(true);
+		Operator = MoveTemp(ConvTransposeOperator);
+	}
+	else IF_MATCHED_INLINABLE_OPERATOR_INITIALIZATION(Cos)
+	else IF_MATCHED_INLINABLE_OPERATOR_INITIALIZATION(Cosh)
+	else IF_MATCHED_MULTIDIRECTIONAL_BROADCASTING_OPERATOR_INITIALIZATION(Div)
+	else IF_MATCHED_INLINABLE_OPERATOR_INITIALIZATION(Exp)
+	else IF_MATCHED_INLINABLE_OPERATOR_INITIALIZATION(Floor)
+	else IF_MATCHED_OPERATOR_WITH_PROTO_INITIALIZATION(Gemm)
+	else IF_MATCHED_INLINABLE_OPERATOR_WITH_PROTO_INITIALIZATION(LeakyRelu, false)
+	else IF_MATCHED_INLINABLE_OPERATOR_INITIALIZATION(Log)
+	else IF_MATCHED_MULTIDIRECTIONAL_BROADCASTING_OPERATOR_INITIALIZATION(Mul)
+	else IF_MATCHED_INLINABLE_OPERATOR_INITIALIZATION(Neg)
+	else IF_MATCHED_MULTIDIRECTIONAL_BROADCASTING_OPERATOR_INITIALIZATION(Pow)
+	else IF_MATCHED_INLINABLE_OPERATOR_INITIALIZATION(Reciprocal)
+	else IF_MATCHED_INLINABLE_OPERATOR_INITIALIZATION(Relu)
+	else IF_MATCHED_INLINABLE_OPERATOR_WITH_PROTO_INITIALIZATION(Reshape, true)
+	else IF_MATCHED_INLINABLE_OPERATOR_INITIALIZATION(Round)
+	else IF_MATCHED_INLINABLE_OPERATOR_INITIALIZATION(Sigmoid)
+	else IF_MATCHED_INLINABLE_OPERATOR_INITIALIZATION(Sign)
+	else IF_MATCHED_INLINABLE_OPERATOR_INITIALIZATION(Sin)
+	else IF_MATCHED_INLINABLE_OPERATOR_INITIALIZATION(Sinh)
+	else IF_MATCHED_INLINABLE_OPERATOR_INITIALIZATION(Sqrt)
+	else IF_MATCHED_MULTIDIRECTIONAL_BROADCASTING_OPERATOR_INITIALIZATION(Sub)
+	else IF_MATCHED_INLINABLE_OPERATOR_INITIALIZATION(Tan)
+	else IF_MATCHED_INLINABLE_OPERATOR_INITIALIZATION(Tanh)
+	if (!Operator.IsValid())
+	{
+		UE_LOG(LogNeuralNetworkInference, Warning, TEXT("FGraphProtoToNeuralNetworkConverter::CreateOperator(): Unknown or unimplemented operator %s."), *InNodeProto.OperatorType);
+	}
+	return Operator;
+}
diff --git a/Engine/Plugins/Experimental/NeuralNetworkInference/Source/NeuralNetworkInference/Private/GraphProtoToNeuralNetworkConverter.h b/Engine/Plugins/Experimental/NeuralNetworkInference/Source/NeuralNetworkInference/Private/GraphProtoToNeuralNetworkConverter.h
new file mode 100644
--- /dev/null
+++ b/Engine/Plugins/Experimental/NeuralNetworkInference/Source/NeuralNetworkInference/Private/GraphProtoToNeuralNetworkConverter.h
@@ -0,0 +1,65 @@
+// Copyright Epic Games, Inc. All Rights Reserved.
+
+#pragma once
+
+#include "CoreMinimal.h"
+#include "ModelProto.h"
+#include "NeuralOperator.h"
+#include "NeuralTensor.h"
+#include "NeuralTensorManager.h"
+
+class FGraphProtoToNeuralNetworkConverter
+{
+public:
+	/**
+	 * It takes a FGraphProto and model path as inputs, and produces an array of FNeuralOperator's and FNeuralTensorManager as result.
+	 * These 2 outputs are used by UNeuralNetworkLegacy to create and load a (deep) neural network architecture and its weights.
+	 * @param InTensorManager It is const, but kept not const for CreateOperatorsAndEditTensorArray() compatibility.
+	 *
+	 * 2 alternatives:
+	 * Without InModelPath, InOutTensorManager is assumed to be a const input variable (assuming the tensors are all already loaded, e.g., on de-serialization).
+	 * With InModelPath (Editor-only function), InOutTensorManager is assumed to be an output variable that will be filled based on InModelPath.
+	 */
+	static bool Translate(TArray<TSharedPtr<FNeuralOperator>>& OutOperators, FNeuralTensorManager& InTensorManager, const FGraphProto& InGraphProto);
+#if WITH_EDITOR
+	static bool Translate(TArray<TSharedPtr<FNeuralOperator>>& OutOperators, FNeuralTensorManager& OutTensorManager, const FGraphProto& InGraphProto, const FString& InModelPath);
+#endif //WITH_EDITOR
+
+private:
+	/**
+	 * Auxiliary function for Translate().
+	 * It initializes both OutOperators and OutputNameIndexMap, as well as modify InOutTensors and InOutNameIndexMap.
+	 * OutputNameIndexMap is just the copy of InOutputNameDummyIndexMap while filling its int32 Value fields. Given the dual Game/Editor behaviors, and to optimize Game mode, kept as 2 different rather than modifying in-place as InOutOutputNameDummyIndexMap.
+	 *
+	 * 2 alternatives:
+	 * - Game and/or Editor (efficient): If InModelPath is empty, InOutTensorManager is assumed to be a const input variable (assuming the tensors are all already loaded, e.g., on de-serialization).
+	 * - Editor-only: If InModelPath is not empty, InOutTensorManager is assumed to be an output variable that will be filled based on InModelPath.
+	 *
+	 * @param OutputNameIndexMap It is created from InOutputNameDummyIndexMap, assuming the FString part of the mapping is right, and it will fill the index (int32 part of it).
+	 * @param InInputNameDummyIndexMap The int32 part is irrelevant, the only important part is whether the FString field is filled or not, regardless of its value. Kept as TMap for performance to integrate better with FTensorManager.
+	 * @param InOutputNameDummyIndexMap The int32 part is irrelevant, the only important part is whether the FString field is filled or not, regardless of its value. Kept as TMap for performance to integrate better with FTensorManager.
+	 */
+	static bool CreateOperatorsAndEditTensorArray(TArray<TSharedPtr<FNeuralOperator>>& OutOperators, TMap<FString, int32>& OutputNameIndexMap, TArray<FNeuralTensor>& InOutTensors,
+		TMap<FString, int32>& InOutNameIndexMap, const TMap<FString, int32>& InInputNameDummyIndexMap, const TMap<FString, int32>& InOutputNameDummyIndexMap,
+		const FGraphProto& InGraphProto, const FString& InModelPath = TEXT(""));
+
+	/**
+	 * It will return FNeuralOperator::EDataType::NotInlined in any of the 2 following cases occur:
+	 * - If the tensor is an EGPUDataType::Output.
+	 * - If the tensor is used by another layer.
+	 * It will return FNeuralOperator::EDataType::Inlined otherwise.
+	 */
+	static bool CanTensorBeInlined(const FString& InTensorName, const TMap<FString, int32>& InInputNameDummyIndexMap,
+		const TMap<FString, int32>& InOutputNameDummyIndexMap, const FNodeProto& InNodeProto, const FGraphProto& InGraphProto, const bool bCanOutputPropertiesChangeWithPostForward = false);
+	/**
+	 * It will return which input tensors can be inlined.
+	 */
+	static TSet<uint32> GetPotentiallyInlinedTensors(const TArray<FString>& InTensorNames, const TMap<FString, int32>& InInputNameDummyIndexMap,
+		const TMap<FString, int32>& InOutputNameDummyIndexMap, const FNodeProto& InNodeProto, const FGraphProto& InGraphProto);
+
+	/**
+	 * Given the input arguments, it creates the proper TShared<FNeuralOperator> child instance.
+	 */
+	static TSharedPtr<FNeuralOperator> CreateOperator(const FNodeProto& InNodeProto, const TArray<FString>& InInputTensorNamesForOperator,
+		const TMap<FString, int32>& InInputNameDummyIndexMap, const TMap<FString, int32>& InOutputNameDummyIndexMap, const FGraphProto& InGraphProto);
+};
diff --git a/Engine/Plugins/Experimental/NeuralNetworkInference/Source/NeuralNetworkInference/Private/NeuralNetworkFromONNXTranslator.cpp b/Engine/Plugins/Experimental/NeuralNetworkInference/Source/NeuralNetworkInference/Private/NeuralNetworkFromONNXTranslator.cpp
deleted file mode 100644
diff --git a/Engine/Plugins/Experimental/NeuralNetworkInference/Source/NeuralNetworkInference/Private/NeuralNetworkFromONNXTranslator.h b/Engine/Plugins/Experimental/NeuralNetworkInference/Source/NeuralNetworkInference/Private/NeuralNetworkFromONNXTranslator.h
deleted file mode 100644
diff --git a/Engine/Plugins/Experimental/NeuralNetworkInference/Source/NeuralNetworkInference/Private/NeuralNetworkLegacy.cpp b/Engine/Plugins/Experimental/NeuralNetworkInference/Source/NeuralNetworkInference/Private/NeuralNetworkLegacy.cpp
--- a/Engine/Plugins/Experimental/NeuralNetworkInference/Source/NeuralNetworkInference/Private/NeuralNetworkLegacy.cpp
+++ b/Engine/Plugins/Experimental/NeuralNetworkInference/Source/NeuralNetworkInference/Private/NeuralNetworkLegacy.cpp
@@ -2,7 +2,7 @@
 
 #include "NeuralNetworkLegacy.h"
 #include "NeuralOperator.h"
-#include "NeuralNetworkFromONNXTranslator.h"
+#include "GraphProtoToNeuralNetworkConverter.h"
 #include "NeuralNetworkInferenceUtils.h"
 #include "NeuralNetworkInferenceVersion.h"
 #include "EditorFramework/AssetImportData.h"
@@ -68,7 +68,7 @@ void UNeuralNetworkLegacy::PostLoad()
 			bIsLoaded = false;
 		}
 		// Turn ModelProto into Operators
-		else if (!FNeuralNetworkFromONNXTranslator::Translate(Operators, TensorManager, ModelProto.GetGraph()))
+		else if (!FGraphProtoToNeuralNetworkConverter::Translate(Operators, TensorManager, ModelProto.GetGraph()))
 		{
 			UE_LOG(LogNeuralNetworkInference, Warning, TEXT("UNeuralNetworkLegacy::PostLoad(): UNeuralNetworkLegacy could not be configured from its FModelProto."));
 		}
@@ -146,7 +146,7 @@ bool UNeuralNetworkLegacy::Load(const FString& InFilePath)
 		return false;
 	}
 	// Turn ModelProto into Operators
-	bIsLoaded = FNeuralNetworkFromONNXTranslator::Translate(Operators, TensorManager, ModelProto.GetGraph(), InFilePath);
+	bIsLoaded = FGraphProtoToNeuralNetworkConverter::Translate(Operators, TensorManager, ModelProto.GetGraph(), InFilePath);
 	if (!TensorManager.IsLoaded() || !FNeuralNetworkInferenceVersion::CheckVersion(TensorManager.GetVersion()))
 	{
 		UE_LOG(LogNeuralNetworkInference, Warning, TEXT("UNeuralNetworkLegacy::Load(): TensorManager could not be loaded from %s or is outdated. IsLoaded() = %d."),
diff --git a/Engine/Plugins/Experimental/NeuralNetworkInference/Source/NeuralNetworkInference/Private/NeuralTensor.cpp b/Engine/Plugins/Experimental/NeuralNetworkInference/Source/NeuralNetworkInference/Private/NeuralTensor.cpp
--- a/Engine/Plugins/Experimental/NeuralNetworkInference/Source/NeuralNetworkInference/Private/NeuralTensor.cpp
+++ b/Engine/Plugins/Experimental/NeuralNetworkInference/Source/NeuralNetworkInference/Private/NeuralTensor.cpp
@@ -431,7 +431,7 @@ bool FNeuralTensor::SetFromTensorProto(const FTensorProto* const InTensorProto,
 {
 	if (!InTensorProto)
 	{
-		UE_LOG(LogNeuralNetworkInference, Warning, TEXT("FNeuralNetworkFromONNXTranslator::SetFromTensorProto(): InTensorProto was a nullptr."));
+		UE_LOG(LogNeuralNetworkInference, Warning, TEXT("FNeuralTensor::SetFromTensorProto(): InTensorProto was a nullptr."));
 		return false;
 	}
 
@@ -454,11 +454,11 @@ bool FNeuralTensor::SetFromTensorProto(const FTensorProto* const InTensorProto,
 	{
 #if WITH_EDITOR
 #if PLATFORM_WINDOWS
-		UE_LOG(LogNeuralNetworkInference, Warning, TEXT("FNeuralNetworkFromONNXTranslator::SetFromTensorProto(): DEPRECATED CODE, OTXT no longer functional."));
+		UE_LOG(LogNeuralNetworkInference, Warning, TEXT("FNeuralTensor::SetFromTensorProto(): DEPRECATED CODE, OTXT no longer functional."));
 		// // Sanity check
 		// if (InTensorProto->ExternalData[0].Key != TEXT("location") || InTensorProto->ExternalData[0].Value.Len() < 1)
 		// {
-		// 	UE_LOG(LogNeuralNetworkInference, Warning, TEXT("FNeuralNetworkFromONNXTranslator::SetFromTensorProto(): InTensorProto->ExternalData[0].Key = %s != \"location\""
+		// 	UE_LOG(LogNeuralNetworkInference, Warning, TEXT("FNeuralTensor::SetFromTensorProto(): InTensorProto->ExternalData[0].Key = %s != \"location\""
 		// 		" || InTensorProto->ExternalData[0].Value.Len() = %d (should be > 0)."), *InTensorProto->ExternalData[0].Key, InTensorProto->ExternalData[0].Value.Len());
 		// 	return false;
 		// }
@@ -466,21 +466,21 @@ bool FNeuralTensor::SetFromTensorProto(const FTensorProto* const InTensorProto,
 		// const FString BinaryWeightFilePath = InExternalDataDirectory / InTensorProto->ExternalData[0].Value;
 		// if (!FModelProtoFileReader::ReadWeightsFromOtxtBinaryFile((char*)GetData(), Num() * FDataType::GetSize(NeuralDataType), BinaryWeightFilePath))
 		// {
-		// 	UE_LOG(LogNeuralNetworkInference, Warning, TEXT("FNeuralNetworkFromONNXTranslator::SetFromTensorProto(): Could not read binary file: %s."), *BinaryWeightFilePath);
+		// 	UE_LOG(LogNeuralNetworkInference, Warning, TEXT("FNeuralTensor::SetFromTensorProto(): Could not read binary file: %s."), *BinaryWeightFilePath);
 		// 	return false;
 		// }
 #else //PLATFORM_WINDOWS
-		UE_LOG(LogNeuralNetworkInference, Warning, TEXT("FNeuralNetworkFromONNXTranslator::SetFromTensorProto(): Only implemented in Windows."));
+		UE_LOG(LogNeuralNetworkInference, Warning, TEXT("FNeuralTensor::SetFromTensorProto(): Only implemented in Windows."));
 		return false;
 #endif //PLATFORM_WINDOWS
 #else //WITH_EDITOR
-		UE_LOG(LogNeuralNetworkInference, Warning, TEXT("FNeuralNetworkFromONNXTranslator::SetFromTensorProto(): InTensorProto->ExternalData should never be used in non-Editor mode."));
+		UE_LOG(LogNeuralNetworkInference, Warning, TEXT("FNeuralTensor::SetFromTensorProto(): InTensorProto->ExternalData should never be used in non-Editor mode."));
 		return false;
 #endif //WITH_EDITOR
 	}
 	else
 	{
-		UE_LOG(LogNeuralNetworkInference, Warning, TEXT("FNeuralNetworkFromONNXTranslator::SetFromTensorProto(): InTensorProto was empty (RawData and ExternalData)."));
+		UE_LOG(LogNeuralNetworkInference, Warning, TEXT("FNeuralTensor::SetFromTensorProto(): InTensorProto was empty (RawData and ExternalData)."));
 		return false;
 	}
 
