NeuralNetworkInference (NNI) plugin:
- More verbose for QA CIS.

#rnx

diff --git a/Engine/Plugins/Experimental/NeuralNetworkInference/Source/NeuralNetworkInferenceQA/Private/LegacyModelUnitTester.cpp b/Engine/Plugins/Experimental/NeuralNetworkInference/Source/NeuralNetworkInferenceQA/Private/LegacyModelUnitTester.cpp
--- a/Engine/Plugins/Experimental/NeuralNetworkInference/Source/NeuralNetworkInferenceQA/Private/LegacyModelUnitTester.cpp
+++ b/Engine/Plugins/Experimental/NeuralNetworkInference/Source/NeuralNetworkInferenceQA/Private/LegacyModelUnitTester.cpp
@@ -94,12 +94,12 @@ UNeuralNetworkLegacy* FLegacyModelUnitTester::NetworkUassetLoadTest(const FStrin
 	UNeuralNetworkLegacy* Network = LoadObject<UNeuralNetworkLegacy>((UObject*)GetTransientPackage(), *InUAssetPath);
 	if (!Network)
 	{
-		ensureMsgf(false, TEXT("UNeuralNetworkLegacy is a nullptr."));
+		ensureMsgf(false, TEXT("UNeuralNetworkLegacy is a nullptr. Path: \"%s\"."), *InUAssetPath);
 		return nullptr;
 	}
 	if (!Network->IsLoaded())
 	{
-		ensureMsgf(false, TEXT("UNeuralNetworkLegacy could not be loaded from uasset disk location."));
+		ensureMsgf(false, TEXT("UNeuralNetworkLegacy could not be loaded from uasset disk location. Path: \"%s\"."), *InUAssetPath);
 		return nullptr;
 	}
 	UE_LOG(LogNeuralNetworkInferenceQA, Display, TEXT("--------------- %d layers loaded"), Network->GetTensors().Num());
@@ -113,12 +113,12 @@ UNeuralNetworkLegacy* FLegacyModelUnitTester::NetworkONNXLoadTest(const FString&
 	UNeuralNetworkLegacy* Network = NewObject<UNeuralNetworkLegacy>((UObject*)GetTransientPackage(), UNeuralNetworkLegacy::StaticClass());
 	if (!Network)
 	{
-		ensureMsgf(false, TEXT("UNeuralNetworkLegacy is a nullptr."));
+		ensureMsgf(false, TEXT("UNeuralNetworkLegacy is a nullptr. Path: \"%s\"."), *InONNXFilePath);
 		return nullptr;
 	}
 	if (!Network->Load(InONNXFilePath))
 	{
-		ensureMsgf(false, TEXT("UNeuralNetworkLegacy could not be loaded from ONNX file disk location: %s"), *InONNXFilePath);
+		ensureMsgf(false, TEXT("UNeuralNetworkLegacy could not be loaded from ONNX file disk location. Path: \"%s\"."), *InONNXFilePath);
 		return nullptr;
 	}
 	return Network;
diff --git a/Engine/Plugins/Experimental/NeuralNetworkInference/Source/NeuralNetworkInferenceQA/Private/ModelUnitTester.cpp b/Engine/Plugins/Experimental/NeuralNetworkInference/Source/NeuralNetworkInferenceQA/Private/ModelUnitTester.cpp
--- a/Engine/Plugins/Experimental/NeuralNetworkInference/Source/NeuralNetworkInferenceQA/Private/ModelUnitTester.cpp
+++ b/Engine/Plugins/Experimental/NeuralNetworkInference/Source/NeuralNetworkInferenceQA/Private/ModelUnitTester.cpp
@@ -112,17 +112,17 @@ FString FModelUnitTester::GetUAssetModelFilePath(const FString& InModelName, con
 	return InModelName + TEXT("'/Game/") + InModelZooRelativeDirectory / InModelName + TEXT("/") + InModelName + TEXT(".") + InModelName + TEXT("'"); // E.g., '/Game/[MODEL_ZOO_DIR]/ExampleNetworkReadable/ExampleNetworkReadable.ExampleNetworkReadable'
 }
 
-UNeuralNetwork* FModelUnitTester::NetworkUassetLoadTest(const FString& InUassetPath)
+UNeuralNetwork* FModelUnitTester::NetworkUassetLoadTest(const FString& InUAssetPath)
 {
-	UNeuralNetwork* Network = LoadObject<UNeuralNetwork>((UObject*)GetTransientPackage(), *InUassetPath);
+	UNeuralNetwork* Network = LoadObject<UNeuralNetwork>((UObject*)GetTransientPackage(), *InUAssetPath);
 	if (!Network)
 	{
-		ensureMsgf(false, TEXT("UNeuralNetwork is a nullptr."));
+		ensureMsgf(false, TEXT("UNeuralNetwork is a nullptr. Path: \"%s\"."), *InUAssetPath);
 		return nullptr;
 	}
 	if (!Network->IsLoaded())
 	{
-		ensureMsgf(false, TEXT("UNeuralNetwork could not be loaded from uasset disk location."));
+		ensureMsgf(false, TEXT("UNeuralNetwork could not be loaded from uasset disk location. Path: \"%s\"."), *InUAssetPath);
 		return nullptr;
 	}
 	return Network;
@@ -135,12 +135,12 @@ UNeuralNetwork* FModelUnitTester::NetworkONNXOrORTLoadTest(const FString& InMode
 	UNeuralNetwork* Network = NewObject<UNeuralNetwork>((UObject*)GetTransientPackage(), UNeuralNetwork::StaticClass());
 	if (!Network)
 	{
-		ensureMsgf(false, TEXT("UNeuralNetwork is a nullptr."));
+		ensureMsgf(false, TEXT("UNeuralNetwork is a nullptr. Path: \"%s\"."), *InModelFilePath);
 		return nullptr;
 	}
 	if (!Network->Load(InModelFilePath))
 	{
-		ensureMsgf(false, TEXT("UNeuralNetwork could not be loaded from ONNX file disk location: %s"), *InModelFilePath);
+		ensureMsgf(false, TEXT("UNeuralNetwork could not be loaded from ONNX file disk location. Path: \"%s\"."), *InModelFilePath);
 		return nullptr;
 	}
 	return Network;
diff --git a/Engine/Plugins/Experimental/NeuralNetworkInference/Source/NeuralNetworkInferenceQA/Private/ModelUnitTester.h b/Engine/Plugins/Experimental/NeuralNetworkInference/Source/NeuralNetworkInferenceQA/Private/ModelUnitTester.h
--- a/Engine/Plugins/Experimental/NeuralNetworkInference/Source/NeuralNetworkInferenceQA/Private/ModelUnitTester.h
+++ b/Engine/Plugins/Experimental/NeuralNetworkInference/Source/NeuralNetworkInferenceQA/Private/ModelUnitTester.h
@@ -23,7 +23,7 @@ private:
 	static FString GetONNXModelFilePath(const FString& InModelZooDirectory, const FString& InModelName);
 	static FString GetORTModelFilePath(const FString& InModelZooDirectory, const FString& InModelName);
 	static FString GetUAssetModelFilePath(const FString& InModelName, const FString& InModelZooRelativeDirectory);
-	static UNeuralNetwork* NetworkUassetLoadTest(const FString& InUassetPath);
+	static UNeuralNetwork* NetworkUassetLoadTest(const FString& InUAssetPath);
 #if WITH_EDITOR
 	static UNeuralNetwork* NetworkONNXOrORTLoadTest(const FString& InModelFilePath);
 #endif //WITH_EDITOR
