NeuralNetworkInference (NNI) plugin:
- NNILegacy merged into NNI step 1.5: Unnecessary #include AssetImportData.h removed

#rnx

diff --git a/Engine/Plugins/Experimental/NeuralNetworkInference/Source/Private/NeuralNetworkInferenceQAEditor/Private/NeuralNetworkInferenceQAAssetFactory.cpp b/Engine/Plugins/Experimental/NeuralNetworkInference/Source/Private/NeuralNetworkInferenceQAEditor/Private/NeuralNetworkInferenceQAAssetFactory.cpp
--- a/Engine/Plugins/Experimental/NeuralNetworkInference/Source/Private/NeuralNetworkInferenceQAEditor/Private/NeuralNetworkInferenceQAAssetFactory.cpp
+++ b/Engine/Plugins/Experimental/NeuralNetworkInference/Source/Private/NeuralNetworkInferenceQAEditor/Private/NeuralNetworkInferenceQAAssetFactory.cpp
@@ -2,7 +2,6 @@
 
 #include "NeuralNetworkInferenceQAAssetFactory.h"
 #include "NeuralNetworkInferenceQAAsset.h"
-#include "EditorFramework/AssetImportData.h"
 #include "Misc/Paths.h"
 
 
diff --git a/Engine/Plugins/Experimental/NeuralNetworkInference/Source/Private/NeuralNetworkInferenceQAEditor/Private/NeuralNetworkInferenceQAAssetTypeActions.cpp b/Engine/Plugins/Experimental/NeuralNetworkInference/Source/Private/NeuralNetworkInferenceQAEditor/Private/NeuralNetworkInferenceQAAssetTypeActions.cpp
--- a/Engine/Plugins/Experimental/NeuralNetworkInference/Source/Private/NeuralNetworkInferenceQAEditor/Private/NeuralNetworkInferenceQAAssetTypeActions.cpp
+++ b/Engine/Plugins/Experimental/NeuralNetworkInference/Source/Private/NeuralNetworkInferenceQAEditor/Private/NeuralNetworkInferenceQAAssetTypeActions.cpp
@@ -3,7 +3,6 @@
 #include "NeuralNetworkInferenceQAAssetTypeActions.h"
 #include "NeuralNetworkInferenceEditorModule.h"
 #include "NeuralNetworkInferenceQAAsset.h"
-#include "EditorFramework/AssetImportData.h"
 
 
 FText FNeuralNetworkInferenceQAAssetTypeActions::GetName() const
