NeuralNetworkInference (NNI) plugin:
- NNILegacy merged into NNI step 2: UNN & UNNLegacy APIs (even) more similar to each other.

#rnx

diff --git a/Engine/Plugins/Experimental/NeuralNetworkInference/Source/NeuralNetworkInference/Internal/NeuralNetworkLegacy.h b/Engine/Plugins/Experimental/NeuralNetworkInference/Source/NeuralNetworkInference/Internal/NeuralNetworkLegacy.h
--- a/Engine/Plugins/Experimental/NeuralNetworkInference/Source/NeuralNetworkInference/Internal/NeuralNetworkLegacy.h
+++ b/Engine/Plugins/Experimental/NeuralNetworkInference/Source/NeuralNetworkInference/Internal/NeuralNetworkLegacy.h
@@ -91,8 +91,8 @@ public:
 	template<typename T>
 	T* GetInputDataPointerMutable();
 	TMap<FString, void*> CreateInputDataPointersMutable();
-	FRDGBufferUAVRef GetInputBufferUAVRef();
-	TMap<FString, FRDGBufferUAVRef> CreateInputBufferUAVRefs();
+	//FRDGBufferUAVRef GetInputBufferUAVRef();
+	//TMap<FString, FRDGBufferUAVRef> CreateInputBufferUAVRefs();
 
 	/**
 	 * It returns the input FNeuralTensor (map) as a read-only object.
@@ -177,6 +177,13 @@ protected:
 	
 	UPROPERTY(VisibleAnywhere, Category = "Neural Network Inference")
 	ENeuralDeviceType OutputDeviceType;
+	
+	/**
+	 * Whether UNeuralNetwork::Run() will block the thread until completed (Synchronous), or whether it will run on a background thread, not blocking the calling thread (Asynchronous).
+	 * See ENeuralNetworkSynchronousMode for more details.
+	 */
+	UPROPERTY(VisibleAnywhere, Category = "Neural Network Inference")
+	ENeuralNetworkSynchronousMode SynchronousMode;
 
 	UPROPERTY(VisibleAnywhere, Category = "Neural Network Inference")
 	FNeuralTensorManager TensorManager; /* It contains a few TArray and TMaps for all FNeuralTensors (Input, Output, Intermediate(Not)Initialized, Weight) */
@@ -189,22 +196,29 @@ protected:
 #endif // WITH_EDITORONLY_DATA
 
 private:
-	// Non-uproperty members
-	bool bAreTensorsInGpu; /* It should always be false when loaded from uasset (FNeuralTensors are not auto-loaded to GPU) */
+	FOnAsyncRunCompleted OnAsyncRunCompletedDelegate;
+
 	/**
-	 * Operators represents the set of operators of the network that need to run on the Forward pass and that might need to run on the PostForward.
+	 * Only for the vanilla back end.
+	 * It should always be false when loaded from uasset (FNeuralTensors are not auto-loaded to GPU)
+	 */
+	bool bAreTensorsInGpu;
+
+	/**
+	 * Only for the vanilla back end.
+	 * Set of operators that the network need to run on the Forward pass and that might need to run on the PostForward pass.
 	 */
 	TArray<TSharedPtr<FNeuralOperator>> Operators;
-	FOnAsyncRunCompleted OnAsyncRunCompletedDelegate;
 
 public:
+	// Internal functions not needed by the user
 	//~UObject interface
 	virtual void PostInitProperties() override;
 	virtual void PostLoad() override;
 	virtual void Serialize(FArchive& Archive) override;
 	//~End of UObject interface
 #if WITH_EDITOR
-/** Editor-only function: Re-import asset with editor data (imported file). */
+	/** Editor-only function: Re-import asset with editor data (imported file). */
 	void ReimportAssetFromEditorData();
 	/** Editor-only function: Importing data and options used for loading the neural network. */
 	UAssetImportData* GetAssetImportData() const;
diff --git a/Engine/Plugins/Experimental/NeuralNetworkInference/Source/NeuralNetworkInference/Private/NeuralNetwork.cpp b/Engine/Plugins/Experimental/NeuralNetworkInference/Source/NeuralNetworkInference/Private/NeuralNetwork.cpp
--- a/Engine/Plugins/Experimental/NeuralNetworkInference/Source/NeuralNetworkInference/Private/NeuralNetwork.cpp
+++ b/Engine/Plugins/Experimental/NeuralNetworkInference/Source/NeuralNetworkInference/Private/NeuralNetwork.cpp
@@ -7,221 +7,10 @@
 #include "HAL/FileManager.h"
 #include "Misc/FileHelper.h"
 #include "Misc/Paths.h"
-#include <numeric> // std::accumulate
 
-//#define WITH_NNI_CPU_NOT_RECOMMENDED // Only for debugging purposes
-
-#include "ThirdPartyWarningDisabler.h"
-NNI_THIRD_PARTY_INCLUDES_START
-#undef check
-#undef TEXT
-#ifdef WITH_FULL_NNI_SUPPORT
-	#include "RedirectCoutAndCerrToUeLog.h"
-
-	#include "onnxruntime/core/session/onnxruntime_cxx_api.h"
-	#ifdef PLATFORM_WIN64
-	#include "onnxruntime/core/providers/dml/dml_provider_factory.h"
-	#endif
-	#ifdef WITH_NNI_CPU_NOT_RECOMMENDED
-	#include "onnxruntime/core/providers/nni_cpu/nni_cpu_provider_factory.h"
-	#endif //WITH_NNI_CPU_NOT_RECOMMENDED
-#endif //WITH_FULL_NNI_SUPPORT
-NNI_THIRD_PARTY_INCLUDES_END
-
-
-
-/* FImpl
- *****************************************************************************/
-
-struct UNeuralNetwork::FImpl
-{
-#ifdef WITH_FULL_NNI_SUPPORT
-	TUniquePtr<Ort::Env> Environment;
-	TUniquePtr<Ort::Session> Session;
-	TUniquePtr<Ort::AllocatorWithDefaultOptions> Allocator;
-	TUniquePtr<Ort::SessionOptions> SessionOptions;
-
-	static bool InitializedAndConfigureMembers(TSharedPtr<FImpl>& InOutImpl, const FString& InModelFullFilePath, const ENeuralDeviceType InDeviceType);
-
-	bool ConfigureMembers(const ENeuralDeviceType InDeviceType);
-
-	void ConfigureTensors(FNeuralTensors& OutTensors, TArray<bool>* InAreInputTensorSizesVariable = nullptr);
-#endif //WITH_FULL_NNI_SUPPORT
-};
-
-#ifdef WITH_FULL_NNI_SUPPORT
-bool UNeuralNetwork::FImpl::InitializedAndConfigureMembers(TSharedPtr<FImpl>& InOutImpl, const FString& InModelFullFilePath, const ENeuralDeviceType InDeviceType)
-{
-	// Initialize InOutImpl
-	if (!InOutImpl.IsValid())
-	{
-		InOutImpl = MakeShared<FImpl>();
-
-		// Set up ORT and create an environment
-		Ort::InitApi();
-		const char* const ModelFullFilePathCharPtr = TCHAR_TO_ANSI(*InModelFullFilePath);
-		InOutImpl->Environment = MakeUnique<Ort::Env>(ORT_LOGGING_LEVEL_WARNING, ModelFullFilePathCharPtr); // Any unique string would work, it does not need to be the file path
-
-		InOutImpl->Allocator = MakeUnique<Ort::AllocatorWithDefaultOptions>();
-	}
-
-	// Configure InOutImpl
-	if (!InOutImpl->ConfigureMembers(InDeviceType))
-	{
-		UE_LOG(LogNeuralNetworkInference, Warning, TEXT("UNeuralNetwork::Load(): ConfigureMembers failed."));
-		return false;
-	}
-
-	return true;
-}
-
-bool UNeuralNetwork::FImpl::ConfigureMembers(const ENeuralDeviceType InDeviceType)
-{
-	// Configure Session
-	SessionOptions = MakeUnique<Ort::SessionOptions>();
-
-	// Configure number threads
-	SessionOptions->SetIntraOpNumThreads(2);
-	// Uncomment if you want to change the priority of the threads, by default is TPri_Normal
-	//SessionOptions->SetPriorityOpThreads(EThreadPriority::TPri_Normal);
-
-	// Configure Provider
-	// GPU
-	if (InDeviceType == ENeuralDeviceType::GPU)
-	{
-#ifdef PLATFORM_WIN64
-		// ORT GPU (Direct ML)
-		SessionOptions->SetGraphOptimizationLevel(GraphOptimizationLevel::ORT_ENABLE_ALL); // ORT_ENABLE_ALL, ORT_ENABLE_EXTENDED, ORT_ENABLE_BASIC, ORT_DISABLE_ALL
-		if (OrtSessionOptionsAppendExecutionProvider_DML(*SessionOptions, 0))
-		{
-			UE_LOG(LogNeuralNetworkInference, Warning, TEXT("Some error occurred."));
-			return false;
-		}
-		return true; // @todo: Remove this line when NNI_HLSL is working
-#else
-		UE_LOG(LogNeuralNetworkInference, Warning, TEXT("UNeuralNetwork: GPU mode only supported in Windows for now. Please, switch to CPU or to Windows."));
-
-		//SessionOptions->SetGraphOptimizationLevel(GraphOptimizationLevel::ORT_DISABLE_ALL); // ORT_ENABLE_ALL, ORT_ENABLE_EXTENDED, ORT_ENABLE_BASIC, ORT_DISABLE_ALL
-		//if (OrtSessionOptionsAppendExecutionProvider_NNI_HLSL(*SessionOptions, 0))
-		//{
-		//	UE_LOG(LogNeuralNetworkInference, Warning, TEXT("Some error occurred."));
-		//	return false;
-		//}
-#endif //PLATFORM_WIN64
-	}
-	// CPU
-	//else // @todo: Uncomment this line when NNI_HLSL is working
-	{
-#ifdef WITH_NNI_CPU_NOT_RECOMMENDED
-		// NNI CPU (Deprecated)
-		SessionOptions->SetGraphOptimizationLevel(GraphOptimizationLevel::ORT_DISABLE_ALL); // ORT_ENABLE_ALL, ORT_ENABLE_EXTENDED, ORT_ENABLE_BASIC, ORT_DISABLE_ALL
-		if (OrtSessionOptionsAppendExecutionProvider_NNI_CPU(*SessionOptions))
-		{
-			UE_LOG(LogNeuralNetworkInference, Warning, TEXT("Some error occurred."));
-			return false;
-		}
-#else
-		// ORT CPU
-		SessionOptions->SetGraphOptimizationLevel(GraphOptimizationLevel::ORT_ENABLE_ALL); // ORT_ENABLE_ALL, ORT_ENABLE_EXTENDED, ORT_ENABLE_BASIC, ORT_DISABLE_ALL
-#endif //ORT_CPU
-	}
-
-	return true;
-}
-
-void UNeuralNetwork::FImpl::ConfigureTensors(FNeuralTensors& OutTensors, TArray<bool>* InAreInputTensorSizesVariable)
-{
-	const bool bIsInput = (InAreInputTensorSizesVariable != nullptr);
-	TArray<const char*> TensorNames;
-	TArray<ENeuralDataType> TensorDataTypes;
-	TArray<TArray<int64>> TensorSizes;
-
-	const uint32 NumberTensors = bIsInput ? Session->GetInputCount() : Session->GetOutputCount();
-	if (InAreInputTensorSizesVariable)
-	{
-		InAreInputTensorSizesVariable->SetNum(NumberTensors);
-	}
-	for (uint32 TensorIndex = 0; TensorIndex < NumberTensors; ++TensorIndex)
-	{
-		// Get node name
-		{
-			const char* TensorName = bIsInput ? Session->GetInputName(TensorIndex, *Allocator) : Session->GetOutputName(TensorIndex, *Allocator);
-			TensorNames.Emplace(TensorName);
-		}
-
-		// Get node type
-		Ort::TypeInfo CurrentTypeInfo = bIsInput ? Session->GetInputTypeInfo(TensorIndex) : Session->GetOutputTypeInfo(TensorIndex);
-
-		Ort::TensorTypeAndShapeInfo CurrentTensorInfo = CurrentTypeInfo.GetTensorTypeAndShapeInfo();
-
-		{
-			ENeuralDataType TensorDataType;
-			{
-				const ONNXTensorElementDataType ONNXTensorElementDataTypeEnum = CurrentTensorInfo.GetElementType();
-				if (ONNXTensorElementDataTypeEnum == ONNX_TENSOR_ELEMENT_DATA_TYPE_FLOAT)
-				{
-					TensorDataType = ENeuralDataType::Float;
-				}
-				//else if (ONNXTensorElementDataTypeEnum == ONNX_TENSOR_ELEMENT_DATA_TYPE_INT32)
-				//{
-				//	TensorDataType = ENeuralDataType::Int32;
-				//}
-				//else if (ONNXTensorElementDataTypeEnum == ONNX_TENSOR_ELEMENT_DATA_TYPE_INT64)
-				//{
-				//	TensorDataType = ENeuralDataType::Int64;
-				//}
-				//else if (ONNXTensorElementDataTypeEnum == ONNX_TENSOR_ELEMENT_DATA_TYPE_UINT32)
-				//{
-				//	TensorDataType = ENeuralDataType::UInt32;
-				//}
-				//else if (ONNXTensorElementDataTypeEnum == ONNX_TENSOR_ELEMENT_DATA_TYPE_UINT64)
-				//{
-				//	TensorDataType = ENeuralDataType::UInt64;
-				//}
-				else
-				{
-					TensorDataType = ENeuralDataType::None;
-					UE_LOG(LogNeuralNetworkInference, Warning, TEXT("ONNXTensorElementDataTypeEnum = %d not implemented yet."), (int32)ONNXTensorElementDataTypeEnum);
-					return;
-				}
-			}
-			TensorDataTypes.Push(TensorDataType);
-		}
-
-		// Get input shapes/dims
-		{
-			TArray<int64> CurrentTensorSizes;
-			{
-				for (const int64_t CurrentTensorSize : CurrentTensorInfo.GetShape())
-				{
-					if (InAreInputTensorSizesVariable)
-					{
-						(*InAreInputTensorSizesVariable)[TensorIndex] |= (CurrentTensorSize < 0);
-					}
-					// Negative (variable) dimensions not implemented yet
-					if (CurrentTensorSize < 0)
-					{
-						CurrentTensorSizes.Push(1);
-						UE_LOG(LogNeuralNetworkInference, Display,
-							TEXT("Negative (i.e., variable) dimensions not allowed yet, hard-coded to 1. Let us know if you really need variable dimensions."
-								" Keep in mind that fixed sizes might allow additional optimizations and speedup of the network during Run()."));
-					}
-					else
-					{
-						CurrentTensorSizes.Push(CurrentTensorSize);
-					}
-				}
-			}
-			TensorSizes.Push(CurrentTensorSizes);
-		}
-
-		CurrentTypeInfo.release();
-
-	}
-
-	OutTensors.SetFromNetwork(TensorNames, TensorDataTypes, TensorSizes);
-}
-#endif //WITH_FULL_NNI_SUPPORT
+// Other files with UNeuralNetwork implementation (to be included only once and after the includes above)
+#include "NeuralNetworkImplBackEndUEAndORT.imp"
+#include "NeuralNetworkImplBackEndUEOnly.imp"
 
 
 
@@ -229,11 +18,16 @@ void UNeuralNetwork::FImpl::ConfigureTensors(FNeuralTensors& OutTensors, TArray<
  *****************************************************************************/
 
 UNeuralNetwork::UNeuralNetwork()
-	: bIsLoaded(false)
-	, DeviceType(ENeuralDeviceType::CPU)
+	: DeviceType(ENeuralDeviceType::CPU)
 	, InputDeviceType(ENeuralDeviceType::CPU)
 	, OutputDeviceType(ENeuralDeviceType::CPU)
 	, SynchronousMode(ENeuralNetworkSynchronousMode::Synchronous)
+#ifdef WITH_FULL_NNI_SUPPORT
+	, BackEnd(ENeuralBackEnd::UEAndORT)
+#else //WITH_FULL_NNI_SUPPORT
+	, BackEnd(ENeuralBackEnd::UEOnly)
+#endif //WITH_FULL_NNI_SUPPORT
+	, bIsLoaded(false)
 {
 }
 
@@ -270,8 +64,8 @@ bool UNeuralNetwork::Load(const FString& InModelFilePath)
 
 	const FRedirectCoutAndCerrToUeLog RedirectCoutAndCerrToUeLog;
 
-	// Initialize and configure Impl
-	if (!UNeuralNetwork::FImpl::InitializedAndConfigureMembers(Impl, ModelFullFilePath, GetDeviceType()))
+	// Initialize and configure ImplBackEndUEAndORT
+	if (!UNeuralNetwork::FImplBackEndUEAndORT::InitializedAndConfigureMembers(ImplBackEndUEAndORT, ModelFullFilePath, GetDeviceType()))
 	{
 		UE_LOG(LogNeuralNetworkInference, Warning, TEXT("UNeuralNetwork::Load(): InitializedAndConfigureMembers failed."));
 		return false;
@@ -299,19 +93,19 @@ bool UNeuralNetwork::Load(const FString& InModelFilePath)
 //		FPaths::Split(ModelFullFilePath, ONNXPathPart, ONNXFilenamePart, ONNXExtensionPart);
 //		const FString OutputORTOptimizedModelPath = ONNXPathPart / ONNXFilenamePart + TEXT(".ort");
 //#ifdef _WIN32
-//		Impl->SessionOptions->SetOptimizedModelFilePath(*OutputORTOptimizedModelPath);
+//		ImplBackEndUEAndORT->SessionOptions->SetOptimizedModelFilePath(*OutputORTOptimizedModelPath);
 //#else
-//		Impl->SessionOptions->SetOptimizedModelFilePath(TCHAR_TO_ANSI(*OutputORTOptimizedModelPath));
+//		ImplBackEndUEAndORT->SessionOptions->SetOptimizedModelFilePath(TCHAR_TO_ANSI(*OutputORTOptimizedModelPath));
 //#endif //_WIN32
 //		// ONNX --> ORT conversion
 //		// This session is just temporarily opened so the ORT model file can be generated
-//		Impl->Session = MakeUnique<Ort::Session>(*Impl->Environment,
+//		ImplBackEndUEAndORT->Session = MakeUnique<Ort::Session>(*ImplBackEndUEAndORT->Environment,
 //#ifdef _WIN32
 //			*ModelFullFilePath,
 //#else
 //			FilePathCharPtr,
 //#endif
-//			*Impl->SessionOptions);
+//			*ImplBackEndUEAndORT->SessionOptions);
 //
 //		// Read model from OutputORTOptimizedModelPath
 //		return Load(OutputORTOptimizedModelPath);
@@ -321,13 +115,13 @@ bool UNeuralNetwork::Load(const FString& InModelFilePath)
 //	{
 //		// Read model from ModelFullFilePath
 //		std::vector<uint8_t> OutputModelReadFromDiskInBytesVector;
-//		Impl->Session = MakeUnique<Ort::Session>(*Impl->Environment,
+//		ImplBackEndUEAndORT->Session = MakeUnique<Ort::Session>(*ImplBackEndUEAndORT->Environment,
 //#ifdef _WIN32
 //			*ModelFullFilePath,
 //#else
 //			FilePathCharPtr,
 //#endif
-//			*Impl->SessionOptions, &OutputModelReadFromDiskInBytesVector);
+//			*ImplBackEndUEAndORT->SessionOptions, &OutputModelReadFromDiskInBytesVector);
 //
 //		// Fill ModelReadFromDiskInBytes
 //		const int32 ArraySize = OutputModelReadFromDiskInBytesVector.size();
@@ -359,8 +153,8 @@ bool UNeuralNetwork::Load()
 
 	const FRedirectCoutAndCerrToUeLog RedirectCoutAndCerrToUeLog;
 
-	// Initialize and configure Impl
-	if (!UNeuralNetwork::FImpl::InitializedAndConfigureMembers(Impl, ModelFullFilePath, GetDeviceType()))
+	// Initialize and configure ImplBackEndUEAndORT
+	if (!UNeuralNetwork::FImplBackEndUEAndORT::InitializedAndConfigureMembers(ImplBackEndUEAndORT, ModelFullFilePath, GetDeviceType()))
 	{
 		UE_LOG(LogNeuralNetworkInference, Warning, TEXT("UNeuralNetwork::Load(): InitializedAndConfigureMembers failed."));
 		return false;
@@ -370,7 +164,7 @@ bool UNeuralNetwork::Load()
 	if (ModelReadFromDiskInBytes.Num() > 0)
 	{
 		// Read model from ModelReadFromDiskInBytesVector
-		Impl->Session = MakeUnique<Ort::Session>(*Impl->Environment, ModelReadFromDiskInBytes.GetData(), ModelReadFromDiskInBytes.Num(), *Impl->SessionOptions);
+		ImplBackEndUEAndORT->Session = MakeUnique<Ort::Session>(*ImplBackEndUEAndORT->Environment, ModelReadFromDiskInBytes.GetData(), ModelReadFromDiskInBytes.Num(), *ImplBackEndUEAndORT->SessionOptions);
 	}
 	// Else
 	else
@@ -379,8 +173,8 @@ bool UNeuralNetwork::Load()
 		return false;
 	}
 
-	Impl->ConfigureTensors(InputTensors, &AreInputTensorSizesVariable);
-	Impl->ConfigureTensors(OutputTensors);
+	ImplBackEndUEAndORT->ConfigureTensors(InputTensors, &AreInputTensorSizesVariable);
+	ImplBackEndUEAndORT->ConfigureTensors(OutputTensors);
 
 	bIsLoaded = true;
 	return bIsLoaded;
@@ -403,8 +197,11 @@ ENeuralDeviceType UNeuralNetwork::GetDeviceType() const
 
 void UNeuralNetwork::SetDeviceType(const ENeuralDeviceType InDeviceType)
 {
-	DeviceType = InDeviceType;
-	Load();
+	if (DeviceType != InDeviceType)
+	{
+		DeviceType = InDeviceType;
+		Load();
+	}
 }
 
 ENeuralDeviceType UNeuralNetwork::GetInputDeviceType() const
@@ -427,6 +224,35 @@ void UNeuralNetwork::SetOutputDeviceType(const ENeuralDeviceType InOutputDeviceT
 	OutputDeviceType = InOutputDeviceType;
 }
 
+ENeuralNetworkSynchronousMode UNeuralNetwork::GetSynchronousMode() const
+{
+	return SynchronousMode;
+}
+
+void UNeuralNetwork::SetSynchronousMode(const ENeuralNetworkSynchronousMode InSynchronousMode)
+{
+	SynchronousMode = InSynchronousMode;
+}
+
+UNeuralNetwork::FOnAsyncRunCompleted& UNeuralNetwork::GetOnAsyncRunCompletedDelegate()
+{
+	return OnAsyncRunCompletedDelegate;
+}
+
+ENeuralBackEnd UNeuralNetwork::GetBackEnd() const
+{
+	return BackEnd;
+}
+
+void UNeuralNetwork::SetBackEnd(const ENeuralBackEnd InBackEnd)
+{
+	if (BackEnd != InBackEnd)
+	{
+		BackEnd = InBackEnd;
+		Load();
+	}
+}
+
 const FNeuralTensor& UNeuralNetwork::GetInputTensor(const int32 InTensorIndex) const
 {
 	return InputTensors.GetTensor(InTensorIndex);
@@ -457,11 +283,6 @@ const FNeuralTensors& UNeuralNetwork::GetOutputTensors() const
 	return OutputTensors;
 }
 
-UNeuralNetwork::FOnAsyncRunCompleted& UNeuralNetwork::GetOnAsyncRunCompletedDelegate()
-{
-	return OnAsyncRunCompletedDelegate;
-}
-
 void UNeuralNetwork::Run()
 {
 #ifdef WITH_FULL_NNI_SUPPORT
@@ -482,7 +303,7 @@ void UNeuralNetwork::Run()
 	if (SynchronousMode == ENeuralNetworkSynchronousMode::Synchronous)
 	{
 		const FRedirectCoutAndCerrToUeLog RedirectCoutAndCerrToUeLog;
-		Impl->Session->Run(Ort::RunOptions{ nullptr },
+		ImplBackEndUEAndORT->Session->Run(Ort::RunOptions{ nullptr },
 			InputTensors.GetTensorNames(), InputTensors.GetONNXRuntimeTensors(), InputTensors.GetNumberTensors(),
 			OutputTensors.GetTensorNames(), OutputTensors.GetONNXRuntimeTensors(), OutputTensors.GetNumberTensors());
 	}
@@ -502,9 +323,41 @@ void UNeuralNetwork::Run()
 
 
 
-/* UNeuralNetwork public functions that should only be used internally (not by the user)
+/* UNeuralNetwork private functions
  *****************************************************************************/
 
+#if WITH_EDITOR
+UAssetImportData* UNeuralNetwork::GetAssetImportData() const
+{
+	return AssetImportData;
+}
+
+UAssetImportData* UNeuralNetwork::GetAndMaybeCreateAssetImportData()
+{
+	// An existing import data object was not found, so make one here.
+	if (!AssetImportData)
+	{
+		AssetImportData = NewObject<UAssetImportData>(this, TEXT("AssetImportData"));
+	}
+	return AssetImportData;
+}
+
+void UNeuralNetwork::ReimportAssetFromEditorData()
+{
+	//Get the re-import filename
+	const FString ImportedFilename = AssetImportData->GetFirstFilename();
+	if (ImportedFilename.Len() > 0)
+	{
+		// Ensure that the file provided by the path exists
+		if (IFileManager::Get().FileSize(*ImportedFilename) != INDEX_NONE)
+		{
+			UE_LOG(LogNeuralNetworkInference, Display, TEXT("Performing atomic reimport of [%s]"), *ImportedFilename);
+			Load(ImportedFilename);
+		}
+	}
+}
+#endif // WITH_EDITOR
+
 void UNeuralNetwork::PostInitProperties()
 {
 #if WITH_EDITORONLY_DATA
@@ -540,35 +393,3 @@ void UNeuralNetwork::Serialize(FArchive& Archive)
 #endif // WITH_EDITORONLY_DATA
 	Super::Serialize(Archive);
 }
-
-#if WITH_EDITOR
-void UNeuralNetwork::ReimportAssetFromEditorData()
-{
-	//Get the re-import filename
-	const FString ImportedFilename = AssetImportData->GetFirstFilename();
-	if (ImportedFilename.Len() > 0)
-	{
-		// Ensure that the file provided by the path exists
-		if (IFileManager::Get().FileSize(*ImportedFilename) != INDEX_NONE)
-		{
-			UE_LOG(LogNeuralNetworkInference, Display, TEXT("Performing atomic reimport of [%s]"), *ImportedFilename);
-			Load(ImportedFilename);
-		}
-	}
-}
-
-UAssetImportData* UNeuralNetwork::GetAssetImportData() const
-{
-	return AssetImportData;
-}
-
-UAssetImportData* UNeuralNetwork::GetAndMaybeCreateAssetImportData()
-{
-	// An existing import data object was not found, so make one here.
-	if (!AssetImportData)
-	{
-		AssetImportData = NewObject<UAssetImportData>(this, TEXT("AssetImportData"));
-	}
-	return AssetImportData;
-}
-#endif // WITH_EDITOR
diff --git a/Engine/Plugins/Experimental/NeuralNetworkInference/Source/NeuralNetworkInference/Private/NeuralNetworkLegacy.cpp b/Engine/Plugins/Experimental/NeuralNetworkInference/Source/NeuralNetworkInference/Private/NeuralNetworkLegacy.cpp
--- a/Engine/Plugins/Experimental/NeuralNetworkInference/Source/NeuralNetworkInference/Private/NeuralNetworkLegacy.cpp
+++ b/Engine/Plugins/Experimental/NeuralNetworkInference/Source/NeuralNetworkInference/Private/NeuralNetworkLegacy.cpp
@@ -242,15 +242,15 @@ TMap<FString, void*> UNeuralNetworkLegacy::CreateInputDataPointersMutable()
 	return TensorManager.CreateInputDataPointersMutable();
 }
 
-FRDGBufferUAVRef UNeuralNetworkLegacy::GetInputBufferUAVRef()
-{
-	return TensorManager.GetInputBufferUAVRef();
-}
-
-TMap<FString, FRDGBufferUAVRef> UNeuralNetworkLegacy::CreateInputBufferUAVRefs()
-{
-	return TensorManager.CreateInputBufferUAVRefs();
-}
+//FRDGBufferUAVRef UNeuralNetworkLegacy::GetInputBufferUAVRef()
+//{
+//	return TensorManager.GetInputBufferUAVRef();
+//}
+//
+//TMap<FString, FRDGBufferUAVRef> UNeuralNetworkLegacy::CreateInputBufferUAVRefs()
+//{
+//	return TensorManager.CreateInputBufferUAVRefs();
+//}
 
 const FNeuralTensor& UNeuralNetworkLegacy::GetInputTensor() const
 {
diff --git a/Engine/Plugins/Experimental/NeuralNetworkInference/Source/NeuralNetworkInference/Public/NeuralNetwork.h b/Engine/Plugins/Experimental/NeuralNetworkInference/Source/NeuralNetworkInference/Public/NeuralNetwork.h
--- a/Engine/Plugins/Experimental/NeuralNetworkInference/Source/NeuralNetworkInference/Public/NeuralNetwork.h
+++ b/Engine/Plugins/Experimental/NeuralNetworkInference/Source/NeuralNetworkInference/Public/NeuralNetwork.h
@@ -6,14 +6,8 @@
 #include "NeuralEnumClasses.h"
 #include "NeuralTensor.h"
 #include "NeuralTensors.h"
-
 #include "NeuralNetwork.generated.h"
 
-class FNeuralOperator;
-class UAssetImportData;
-struct FModelProto;
-struct FNeuralTensorManager;
-
 /**
  * Whether UNeuralNetwork::Run() will block the thread until completed (Synchronous), or whether it will run on a background thread, not blocking the calling thread (Asynchronous).
  */
@@ -103,7 +97,7 @@ public:
 
 	/**
 	 * Getter and setter functions for DeviceType, InputDeviceType, and OutputDeviceType.
-	 * See DeviceType, InputDeviceType, and OutputDeviceType for more details.
+	 * @see DeviceType, InputDeviceType, and OutputDeviceType for more details.
 	 */
 	ENeuralDeviceType GetDeviceType() const;
 	void SetDeviceType(const ENeuralDeviceType InDeviceType);
@@ -112,6 +106,29 @@ public:
 	ENeuralDeviceType GetOutputDeviceType() const;
 	void SetOutputDeviceType(const ENeuralDeviceType InOutputDeviceType);
 
+	/**
+	 * Getter and setter functions for SynchronousMode.
+	 * @see SynchronousMode and GetOnAsyncRunCompletedDelegate() for more details.
+	 */
+	ENeuralNetworkSynchronousMode GetSynchronousMode() const;
+	void SetSynchronousMode(const ENeuralNetworkSynchronousMode InSynchronousMode);
+
+
+	/**
+	 * GetOnAsyncRunCompletedDelegate() returns a FOnAsyncRunCompleted delegate that will be called when async UNeuralNetwork::Run() is completed.
+	 * This FOnAsyncRunCompleted delegate could be triggered from any thread but will only be triggered if SynchronousMode == ENeuralNetworkSynchronousMode::Asynchronous.
+	 * If SynchronousMode == ENeuralNetworkSynchronousMode::Synchronous, UNeuralNetwork::Run() will block the calling thread until completed, so a callback delegate is not required.
+	 */
+	DECLARE_DELEGATE(FOnAsyncRunCompleted);
+	FOnAsyncRunCompleted& GetOnAsyncRunCompletedDelegate();
+
+	/**
+	 * Getter and setter functions for BackEnd.
+	 * @see ENeuralBackEnd for more details.
+	 */
+	ENeuralBackEnd GetBackEnd() const;
+	void SetBackEnd(const ENeuralBackEnd InBackEnd);
+
 	/**
 	 * Functions to get/fill input.
 	 */
@@ -126,16 +143,6 @@ public:
 	const FNeuralTensor& GetOutputTensor(const int32 InTensorIndex = 0) const;
 	const FNeuralTensors& GetOutputTensors() const;
 
-	/**
-	 * Callback when UNeuralNetwork::Run() is finished. It will only occur if SynchronousMode == ENeuralNetworkSynchronousMode::Asynchronous and it could be triggered from any thread.
-	 */
-	DECLARE_DELEGATE(FOnAsyncRunCompleted);
-
-	/**
-	 * Returns the delegate called when async Run() is over (on any thread)
-	 */
-	FOnAsyncRunCompleted& GetOnAsyncRunCompletedDelegate();
-
 	/**
 	 * Run() executes the forward pass on the current UNeuralNetwork given the current input FDeprecatedNeuralTensor(s), which were previously filled with
 	 * SetInputFromArrayCopy() or GetInputDataPointerMutable().
@@ -145,9 +152,6 @@ public:
 	void Run();
 
 protected:
-	UPROPERTY(VisibleAnywhere, Category = "Neural Network Inference")
-	bool bIsLoaded;
-
 	/**
 	 * Whether Run() will use CPU or GPU acceleration hardware.
 	 * If SetDeviceType() is never called, the default device (EDeviceType::CPU) will be used.
@@ -169,10 +173,16 @@ protected:
 	
 	/**
 	 * Whether UNeuralNetwork::Run() will block the thread until completed (Synchronous), or whether it will run on a background thread, not blocking the calling thread (Asynchronous).
-	 * See ENeuralNetworkSynchronousMode for more details.
+	 * @see ENeuralNetworkSynchronousMode for more details.
 	 */
 	UPROPERTY(VisibleAnywhere, Category = "Neural Network Inference")
 	ENeuralNetworkSynchronousMode SynchronousMode;
+	
+	/**
+	 * @see ENeuralBackEnd for more details.
+	 */
+	UPROPERTY(VisibleAnywhere, Category = "Neural Network Inference")
+	ENeuralBackEnd BackEnd;
 
 	/**
 	 * Original model file path from which this UNeuralNetwork was loaded from.
@@ -193,37 +203,58 @@ protected:
 	UPROPERTY(VisibleAnywhere, Category = "Neural Network Inference")
 	TArray<bool> AreInputTensorSizesVariable;
 
-#if WITH_EDITORONLY_DATA
-	/** Importing data and options used for loading the neural network. */
-	UPROPERTY(VisibleAnywhere, Instanced, Category = "ImportSettings")
-	UAssetImportData* AssetImportData;
-#endif // WITH_EDITORONLY_DATA
-
 private:
+	bool bIsLoaded;
+
 	UPROPERTY()
 	TArray<uint8> ModelReadFromDiskInBytes;
 
+	/**
+	 * @see FOnAsyncRunCompleted and GetOnAsyncRunCompletedDelegate to understand OnAsyncRunCompletedDelegate.
+	 */
 	FOnAsyncRunCompleted OnAsyncRunCompletedDelegate;
 
 	/**
-	 * PIMPL idiom to hide 3rd party dependencies
+	 * Struct pointer containing the UE-and-ORT-based back end implementation.
+	 * PIMPL idiom to minimize memory when not using this back end and to to hide 3rd party dependencies.
 	 * http://www.cppsamples.com/common-tasks/pimpl.html
 	 */
-	struct FImpl;
-	TSharedPtr<FImpl> Impl;
+	struct FImplBackEndUEAndORT;
+	TSharedPtr<FImplBackEndUEAndORT> ImplBackEndUEAndORT;
+
+	/**
+	 * Struct pointer containing the only-UE-based back end implementation.
+	 * PIMPL idiom to minimize memory when not using this back end.
+	 * http://www.cppsamples.com/common-tasks/pimpl.html
+	 */
+	struct FImplBackEndUEOnly;
+	TSharedPtr<FImplBackEndUEOnly> ImplBackEndUEOnly;
 
 public:
-	// Internal functions not needed by user
-	 //~UObject interface
-	virtual void PostInitProperties() override;
-	virtual void PostLoad() override;
-	virtual void Serialize(FArchive& Archive) override;
-	//~End of UObject interface
+#if WITH_EDITOR
+	/**
+	 * Internal and Editor-only functions not needed by the user.
+	 * Importing data and options used for loading the network.
+	 */
+	class UAssetImportData* GetAssetImportData() const;
+	class UAssetImportData* GetAndMaybeCreateAssetImportData();
+#endif // WITH_EDITOR
+
+private:
+#if WITH_EDITORONLY_DATA
+	/** Importing data and options used for loading the network. */
+	UPROPERTY(Instanced)
+	class UAssetImportData* AssetImportData;
+#endif // WITH_EDITORONLY_DATA
+
 #if WITH_EDITOR
 	/** Editor-only function: Re-import asset with editor data (imported file). */
 	void ReimportAssetFromEditorData();
-	/** Editor-only function: Importing data and options used for loading the neural network. */
-	UAssetImportData* GetAssetImportData() const;
-	UAssetImportData* GetAndMaybeCreateAssetImportData();
 #endif // WITH_EDITOR
+
+	//~UObject interface
+	virtual void PostInitProperties() override;
+	virtual void PostLoad() override;
+	virtual void Serialize(FArchive& Archive) override;
+	//~End of UObject interface
 };
diff --git a/Engine/Plugins/Experimental/NeuralNetworkInference/Source/NeuralNetworkInferenceCore/Public/NeuralEnumClasses.h b/Engine/Plugins/Experimental/NeuralNetworkInference/Source/NeuralNetworkInferenceCore/Public/NeuralEnumClasses.h
--- a/Engine/Plugins/Experimental/NeuralNetworkInference/Source/NeuralNetworkInferenceCore/Public/NeuralEnumClasses.h
+++ b/Engine/Plugins/Experimental/NeuralNetworkInference/Source/NeuralNetworkInferenceCore/Public/NeuralEnumClasses.h
@@ -7,12 +7,15 @@
 #include <type_traits>
 #include "NeuralEnumClasses.generated.h"
 
+/**
+ * Whether UNeuralNetwork will use the highly optimized UnrealEngine-and-ONNXRuntime-based back end (UEAndORT) or the less optimized but fully cross platform only-UE one (UEOnly).
+ * We recommend using the UnrealEngine-and-ONNXRuntime-based back end if possible, and only falling back to the only-UE one for those platforms that do not support the ONNXRuntime one yet.
+ */
 UENUM()
-enum class ENeuralDeviceType : uint8
+enum class ENeuralBackEnd : uint8
 {
-	CPU,
-	GPU,
-	None
+	UEAndORT, /* The best of both worlds, UnrealEngine-and-ONNXRuntime-accelerated back end. Recommended for those platforms that support it (e.g., Windows). */
+	UEOnly /* It might be slower than the UEAndORT back end, but it will compile in all platforms and OSs compatible with Unreal Engine. */
 };
 
 UENUM()
@@ -30,6 +33,14 @@ enum class ENeuralDataType : uint8
 	None
 };
 
+UENUM()
+enum class ENeuralDeviceType : uint8
+{
+	CPU,
+	GPU,
+	None
+};
+
 class NEURALNETWORKINFERENCECORE_API FDataType
 {
 public:
