NeuralNetworkInference (NNI) plugin:
- Removed unused no-longer-used FNeuralNetworkInferenceVersion.

#rnx

diff --git a/Engine/Plugins/Experimental/NeuralNetworkInference/Source/Internal/ModelProto/Internal/ModelProto.h b/Engine/Plugins/Experimental/NeuralNetworkInference/Source/Internal/ModelProto/Internal/ModelProto.h
--- a/Engine/Plugins/Experimental/NeuralNetworkInference/Source/Internal/ModelProto/Internal/ModelProto.h
+++ b/Engine/Plugins/Experimental/NeuralNetworkInference/Source/Internal/ModelProto/Internal/ModelProto.h
@@ -872,9 +872,6 @@ struct MODELPROTO_API FModelProto
 	GENERATED_BODY()
 
 public:
-	UPROPERTY(VisibleAnywhere, Category = "Neural Network Inference")
-	TArray<int32> Version;
-
 	UPROPERTY(VisibleAnywhere, Category = "Neural Network Inference")
 	bool bIsLoaded;
 
@@ -965,7 +962,6 @@ public:
 	FModelProto();
 
 	FORCEINLINE bool IsLoaded() { return bIsLoaded; }
-	FORCEINLINE TArray<int32> GetVersion() { return Version; }
 
 	bool LoadFromString(const FString& InProtoString, const int32 InLevel = 0);
 
diff --git a/Engine/Plugins/Experimental/NeuralNetworkInference/Source/Internal/ModelProto/Private/ModelProto.cpp b/Engine/Plugins/Experimental/NeuralNetworkInference/Source/Internal/ModelProto/Private/ModelProto.cpp
--- a/Engine/Plugins/Experimental/NeuralNetworkInference/Source/Internal/ModelProto/Private/ModelProto.cpp
+++ b/Engine/Plugins/Experimental/NeuralNetworkInference/Source/Internal/ModelProto/Private/ModelProto.cpp
@@ -3,7 +3,6 @@
 #include "ModelProto.h"
 #include "ModelProtoUtils.h"
 #include "ModelProtoStringParser.h"
-#include "NeuralNetworkInferenceVersion.h"
 
 
 
@@ -1021,11 +1020,7 @@ bool FModelProto::LoadFromString(const FString& InProtoString, const int32 InLev
 	bIsLoaded &= FModelProtoStringParser::GetModelProtoArray(TrainingInfo, ProtoMap, TEXT("training_info"), InLevel + 1);
 
 	bIsLoaded &= IRVersion > -1 && OpsetImport.Num() > 0;
-	if (bIsLoaded)
-	{
-		Version = FNeuralNetworkInferenceVersion::GetVersion();
-	}
-	else
+	if (!bIsLoaded)
 	{
 		UE_LOG(LogModelProto, Warning, TEXT("FModelProto::LoadFromString() failed."));
 	}
diff --git a/Engine/Plugins/Experimental/NeuralNetworkInference/Source/NeuralNetworkInference/Internal/NeuralNetworkLegacy.h b/Engine/Plugins/Experimental/NeuralNetworkInference/Source/NeuralNetworkInference/Internal/NeuralNetworkLegacy.h
--- a/Engine/Plugins/Experimental/NeuralNetworkInference/Source/NeuralNetworkInference/Internal/NeuralNetworkLegacy.h
+++ b/Engine/Plugins/Experimental/NeuralNetworkInference/Source/NeuralNetworkInference/Internal/NeuralNetworkLegacy.h
@@ -154,8 +154,6 @@ public:
 	FString ToString() const;
 
 protected:
-	UPROPERTY(VisibleAnywhere, Category = "Neural Network Inference")
-	TArray<int32> Version;
 	UPROPERTY(VisibleAnywhere, Category = "Neural Network Inference")
 	bool bIsLoaded;
 	
diff --git a/Engine/Plugins/Experimental/NeuralNetworkInference/Source/NeuralNetworkInference/Internal/NeuralTensorManager.h b/Engine/Plugins/Experimental/NeuralNetworkInference/Source/NeuralNetworkInference/Internal/NeuralTensorManager.h
--- a/Engine/Plugins/Experimental/NeuralNetworkInference/Source/NeuralNetworkInference/Internal/NeuralTensorManager.h
+++ b/Engine/Plugins/Experimental/NeuralNetworkInference/Source/NeuralNetworkInference/Internal/NeuralTensorManager.h
@@ -53,7 +53,6 @@ public:
 	 * It returns whether it loaded the UNeuralNetworkLegacy successfully.
 	 */
 	bool IsLoaded() const;
-	TArray<int32> GetVersion() const;
 
 	/**
 	 * Getter for the Tensors TArray. GetTensors() will return a const reference while GetTensorsMutable() allows modifying it.
@@ -139,8 +138,6 @@ public:
 	TMap<FString, FNeuralTensor> CreateOutputTensorMap() const;
 
 protected:
-	UPROPERTY(VisibleAnywhere, Category = "Neural Network Inference")
-	TArray<int32> Version;
 	UPROPERTY(VisibleAnywhere, Category = "Neural Network Inference")
 	bool bIsLoaded;
 	/**
diff --git a/Engine/Plugins/Experimental/NeuralNetworkInference/Source/NeuralNetworkInference/Private/NeuralNetwork.cpp b/Engine/Plugins/Experimental/NeuralNetworkInference/Source/NeuralNetworkInference/Private/NeuralNetwork.cpp
--- a/Engine/Plugins/Experimental/NeuralNetworkInference/Source/NeuralNetworkInference/Private/NeuralNetwork.cpp
+++ b/Engine/Plugins/Experimental/NeuralNetworkInference/Source/NeuralNetworkInference/Private/NeuralNetwork.cpp
@@ -2,7 +2,6 @@
 
 #include "NeuralNetwork.h"
 #include "NeuralNetworkInferenceUtils.h"
-#include "NeuralNetworkInferenceVersion.h"
 #include "EditorFramework/AssetImportData.h"
 #include "HAL/FileManager.h"
 #include "Misc/FileHelper.h"
diff --git a/Engine/Plugins/Experimental/NeuralNetworkInference/Source/NeuralNetworkInference/Private/NeuralNetworkLegacy.cpp b/Engine/Plugins/Experimental/NeuralNetworkInference/Source/NeuralNetworkInference/Private/NeuralNetworkLegacy.cpp
--- a/Engine/Plugins/Experimental/NeuralNetworkInference/Source/NeuralNetworkInference/Private/NeuralNetworkLegacy.cpp
+++ b/Engine/Plugins/Experimental/NeuralNetworkInference/Source/NeuralNetworkInference/Private/NeuralNetworkLegacy.cpp
@@ -4,7 +4,6 @@
 #include "NeuralOperator.h"
 #include "GraphProtoToNeuralNetworkConverter.h"
 #include "NeuralNetworkInferenceUtils.h"
-#include "NeuralNetworkInferenceVersion.h"
 #include "EditorFramework/AssetImportData.h"
 #include "HAL/FileManager.h"
 #include "RenderGraphBuilder.h"
@@ -62,11 +61,6 @@ void UNeuralNetworkLegacy::PostLoad()
 				IsLoaded(), TensorManager.IsLoaded(), ModelProto.IsLoaded());
 			bIsLoaded = false;
 		}
-		else if (!FNeuralNetworkInferenceVersion::CheckVersion(Version) || !FNeuralNetworkInferenceVersion::CheckVersion(TensorManager.GetVersion()) || !FNeuralNetworkInferenceVersion::CheckVersion(ModelProto.GetVersion()))
-		{
-			UE_LOG(LogNeuralNetworkInference, Warning, TEXT("UNeuralNetworkLegacy::PostLoad(): CheckVersion() failed."));
-			bIsLoaded = false;
-		}
 		// Turn ModelProto into Operators
 		else if (!FGraphProtoToNeuralNetworkConverter::Translate(Operators, TensorManager, ModelProto.GetGraph()))
 		{
@@ -139,25 +133,18 @@ bool UNeuralNetworkLegacy::Load(const FString& InFilePath)
 	}
 
 	// Read ModelProto
-	if (!FModelProtoFileReader::ReadModelProtoFromFile(ModelProto, InFilePath) || !ModelProto.IsLoaded() || !FNeuralNetworkInferenceVersion::CheckVersion(ModelProto.GetVersion()))
+	if (!FModelProtoFileReader::ReadModelProtoFromFile(ModelProto, InFilePath) || !ModelProto.IsLoaded())
 	{
-		UE_LOG(LogNeuralNetworkInference, Warning, TEXT("UNeuralNetworkLegacy::Load(): Model could not be loaded from %s or is outdated. IsLoaded() = %d."),
-			*InFilePath, ModelProto.IsLoaded());
+		UE_LOG(LogNeuralNetworkInference, Warning, TEXT("UNeuralNetworkLegacy::Load(): Model could not be loaded from %s."), *InFilePath);
 		return false;
 	}
 	// Turn ModelProto into Operators
 	bIsLoaded = FGraphProtoToNeuralNetworkConverter::Translate(Operators, TensorManager, ModelProto.GetGraph(), InFilePath);
-	if (!TensorManager.IsLoaded() || !FNeuralNetworkInferenceVersion::CheckVersion(TensorManager.GetVersion()))
+	if (!TensorManager.IsLoaded())
 	{
-		UE_LOG(LogNeuralNetworkInference, Warning, TEXT("UNeuralNetworkLegacy::Load(): TensorManager could not be loaded from %s or is outdated. IsLoaded() = %d."),
-			*InFilePath, TensorManager.IsLoaded());
+		UE_LOG(LogNeuralNetworkInference, Warning, TEXT("UNeuralNetworkLegacy::Load(): TensorManager could not be loaded from %s."), *InFilePath);
 		return false;
 	}
-	// Update Version
-	if (bIsLoaded)
-	{
-		Version = FNeuralNetworkInferenceVersion::GetVersion();
-	}
 	return bIsLoaded;
 
 #else //PLATFORM_WINDOWS
@@ -171,19 +158,13 @@ bool UNeuralNetworkLegacy::Load(const FString& InFilePath)
 bool UNeuralNetworkLegacy::Load(FNeuralTensorManager& InTensorManager, const TArray<TSharedPtr<FNeuralOperator>>& InOperators)
 {
 	DECLARE_SCOPE_CYCLE_COUNTER(TEXT("UNeuralNetworkLegacy_Load"), STAT_UNeuralNetworkLegacy_Load, STATGROUP_MachineLearning);
-	if (!InTensorManager.IsLoaded() || !FNeuralNetworkInferenceVersion::CheckVersion(InTensorManager.GetVersion()))
+	if (!InTensorManager.IsLoaded())
 	{
-		UE_LOG(LogNeuralNetworkInference, Warning, TEXT("UNeuralNetworkLegacy::Load(): TensorManager could not be loaded or is outdated. IsLoaded() = %d."),
-			InTensorManager.IsLoaded());
+		UE_LOG(LogNeuralNetworkInference, Warning, TEXT("UNeuralNetworkLegacy::Load(): TensorManager could not be loaded."));
 	}
 	Swap(TensorManager, InTensorManager);
 	Operators = InOperators;
 	bIsLoaded = (Operators.Num() > 0 && TensorManager.IsLoaded());
-	// Update Version
-	if (bIsLoaded)
-	{
-		Version = FNeuralNetworkInferenceVersion::GetVersion();
-	}
 	return bIsLoaded;
 }
 
diff --git a/Engine/Plugins/Experimental/NeuralNetworkInference/Source/NeuralNetworkInference/Private/NeuralTensorManager.cpp b/Engine/Plugins/Experimental/NeuralNetworkInference/Source/NeuralNetworkInference/Private/NeuralTensorManager.cpp
--- a/Engine/Plugins/Experimental/NeuralNetworkInference/Source/NeuralNetworkInference/Private/NeuralTensorManager.cpp
+++ b/Engine/Plugins/Experimental/NeuralNetworkInference/Source/NeuralNetworkInference/Private/NeuralTensorManager.cpp
@@ -2,7 +2,6 @@
 
 #include "NeuralTensorManager.h"
 #include "NeuralNetworkInferenceUtils.h"
-#include "NeuralNetworkInferenceVersion.h"
 
 
 
@@ -140,8 +139,6 @@ bool FNeuralTensorManager::Load(TArray<FNeuralTensor>& InTensors, TMap<FString,
 	}
 	// FNeuralTensorManager was configured
 	bIsLoaded = true;
-	// Update Version
-	Version = FNeuralNetworkInferenceVersion::GetVersion();
 	return bIsLoaded;
 }
 
@@ -176,11 +173,6 @@ bool FNeuralTensorManager::IsLoaded() const
 	return bIsLoaded;
 }
 
-TArray<int32> FNeuralTensorManager::GetVersion() const
-{
-	return Version;
-}
-
 const TArray<FNeuralTensor>& FNeuralTensorManager::GetTensors() const
 {
 	SanityCheckIsLoaded();
diff --git a/Engine/Plugins/Experimental/NeuralNetworkInference/Source/NeuralNetworkInferenceCore/Private/NeuralNetworkInferenceVersion.cpp b/Engine/Plugins/Experimental/NeuralNetworkInference/Source/NeuralNetworkInferenceCore/Private/NeuralNetworkInferenceVersion.cpp
deleted file mode 100644
diff --git a/Engine/Plugins/Experimental/NeuralNetworkInference/Source/NeuralNetworkInferenceCore/Public/NeuralNetworkInferenceVersion.h b/Engine/Plugins/Experimental/NeuralNetworkInference/Source/NeuralNetworkInferenceCore/Public/NeuralNetworkInferenceVersion.h
deleted file mode 100644
diff --git a/Engine/Plugins/Experimental/NeuralNetworkInference/Source/Private/ModelProtoFileReader/Private/ModelProtoConverter.cpp b/Engine/Plugins/Experimental/NeuralNetworkInference/Source/Private/ModelProtoFileReader/Private/ModelProtoConverter.cpp
--- a/Engine/Plugins/Experimental/NeuralNetworkInference/Source/Private/ModelProtoFileReader/Private/ModelProtoConverter.cpp
+++ b/Engine/Plugins/Experimental/NeuralNetworkInference/Source/Private/ModelProtoFileReader/Private/ModelProtoConverter.cpp
@@ -1,7 +1,6 @@
 // Copyright Epic Games, Inc. All Rights Reserved.
 
 #include "ModelProtoConverter.h"
-#include "NeuralNetworkInferenceVersion.h"
 
 
 
@@ -32,7 +31,6 @@ bool FModelProtoConverter::ConvertProto3ToUAsset(FModelProto& OutModelProto, con
 {
 	// FModelProto from Protobuf onnx::ModelProto
 	OutModelProto = FModelProto();
-	OutModelProto.Version = FNeuralNetworkInferenceVersion::GetVersion();
 	OutModelProto.bIsLoaded = true;
 	OutModelProto.IRVersion = InONNXModelProto.ir_version();
 	// OperatorSetIds
diff --git a/Engine/Plugins/Experimental/NeuralNetworkInference/Source/Private/ModelProtoFileReader/Private/ModelProtoFileReader.cpp b/Engine/Plugins/Experimental/NeuralNetworkInference/Source/Private/ModelProtoFileReader/Private/ModelProtoFileReader.cpp
--- a/Engine/Plugins/Experimental/NeuralNetworkInference/Source/Private/ModelProtoFileReader/Private/ModelProtoFileReader.cpp
+++ b/Engine/Plugins/Experimental/NeuralNetworkInference/Source/Private/ModelProtoFileReader/Private/ModelProtoFileReader.cpp
@@ -2,7 +2,6 @@
 
 #include "ModelProtoFileReader.h"
 #include "ModelProtoFileReaderUtils.h"
-#include "NeuralNetworkInferenceVersion.h"
 #include "Interfaces/IPluginManager.h"
 #include "Misc/FileHelper.h"
 #include "Misc/Paths.h"
