NeuralNetworkInference (NNI) plugin:
- Improved doc in UNN

#rnx

diff --git a/Engine/Plugins/Experimental/NeuralNetworkInference/Source/NeuralNetworkInference/Public/NeuralNetwork.h b/Engine/Plugins/Experimental/NeuralNetworkInference/Source/NeuralNetworkInference/Public/NeuralNetwork.h
--- a/Engine/Plugins/Experimental/NeuralNetworkInference/Source/NeuralNetworkInference/Public/NeuralNetwork.h
+++ b/Engine/Plugins/Experimental/NeuralNetworkInference/Source/NeuralNetworkInference/Public/NeuralNetwork.h
@@ -131,6 +131,7 @@ public:
 
 	/**
 	 * Functions to get/fill input.
+	 * These functions either take a TArray as input, return a modificable void* to fill the data, or return a constant FNeuralTensor(s) to see input properties (e.g., size or dimensions).
 	 */
 	const FNeuralTensor& GetInputTensor(const int32 InTensorIndex = 0) const;
 	const FNeuralTensors& GetInputTensors() const;
@@ -138,7 +139,7 @@ public:
 	void* GetInputDataPointerMutable(const int32 InTensorIndex = 0);
 
 	/**
-	 * Functions to get output.
+	 * Functions to get output. The returned FNeuralTensor(s) are constant to prevent the user from modifying the tensor properties (e.g., size or dimensions).
 	 */
 	const FNeuralTensor& GetOutputTensor(const int32 InTensorIndex = 0) const;
 	const FNeuralTensors& GetOutputTensors() const;
