NeuralNetworkInference (NNI) plugin:
- ENeuralBackEnd adds Auto for simplicity.

#rnx

diff --git a/Engine/Plugins/Experimental/NeuralNetworkInference/Source/NeuralNetworkInferenceCore/Public/NeuralEnumClasses.h b/Engine/Plugins/Experimental/NeuralNetworkInference/Source/NeuralNetworkInferenceCore/Public/NeuralEnumClasses.h
--- a/Engine/Plugins/Experimental/NeuralNetworkInference/Source/NeuralNetworkInferenceCore/Public/NeuralEnumClasses.h
+++ b/Engine/Plugins/Experimental/NeuralNetworkInference/Source/NeuralNetworkInferenceCore/Public/NeuralEnumClasses.h
@@ -9,13 +9,14 @@
 
 /**
  * Whether UNeuralNetwork will use the highly optimized UnrealEngine-and-ONNXRuntime-based back end (UEAndORT) or the less optimized but fully cross platform only-UE one (UEOnly).
- * We recommend using the UnrealEngine-and-ONNXRuntime-based back end if possible, and only falling back to the only-UE one for those platforms that do not support the ONNXRuntime one yet.
+ * We recommend using Auto, which will find and use the optimal back end for each platform.
  */
 UENUM()
 enum class ENeuralBackEnd : uint8
 {
-	UEAndORT, /* The best of both worlds, UnrealEngine-and-ONNXRuntime-accelerated back end. Recommended for those platforms that support it (e.g., Windows). */
-	UEOnly /* It might be slower than the UEAndORT back end, but it will compile in all platforms and OSs compatible with Unreal Engine. */
+	UEAndORT, /* UnrealEngine-and-ONNXRuntime-accelerated back end, ideal for those platforms that support it. */
+	UEOnly, /* It might be slower than the UEAndORT back end, but it will compile in all platforms and OSs compatible with Unreal Engine. */
+	Auto /* Recommended setting. It will use UEAndORT if possible, and fall back to UEOnly only for those platforms that do not support UEAndORT yet. */
 };
 
 UENUM()
