NeuralNetworkInference (NNI) plugin:
- ModelProtoFileReader enabled for Linux & Mac.
- Fixed compiler errors for ModelProtoFileReader in Linux.
- ConvertProto3ToUAssetBasicType now allows 2 different input types (as long as same sizeof).

#rnx

diff --git a/Engine/Plugins/Experimental/NeuralNetworkInference/NeuralNetworkInference.uplugin b/Engine/Plugins/Experimental/NeuralNetworkInference/NeuralNetworkInference.uplugin
--- a/Engine/Plugins/Experimental/NeuralNetworkInference/NeuralNetworkInference.uplugin
+++ b/Engine/Plugins/Experimental/NeuralNetworkInference/NeuralNetworkInference.uplugin
@@ -66,10 +66,7 @@
 		{
 			"Name": "ModelProtoFileReader",
 			"Type": "Editor",
-			"LoadingPhase": "Default",
-			"WhitelistPlatforms": [
-				"Win64"
-			]
+			"LoadingPhase": "Default"
 		}
 	]
 }
\ No newline at end of file
diff --git a/Engine/Plugins/Experimental/NeuralNetworkInference/Source/Private/ModelProtoFileReader/Private/ModelProtoConverter.h b/Engine/Plugins/Experimental/NeuralNetworkInference/Source/Private/ModelProtoFileReader/Private/ModelProtoConverter.h
--- a/Engine/Plugins/Experimental/NeuralNetworkInference/Source/Private/ModelProtoFileReader/Private/ModelProtoConverter.h
+++ b/Engine/Plugins/Experimental/NeuralNetworkInference/Source/Private/ModelProtoFileReader/Private/ModelProtoConverter.h
@@ -63,11 +63,11 @@ private:
 
 	static bool ConvertProto3ToUAssetUInt8(TArray<uint8>& OutDataArray, const std::string& InRawDataString);
 
-	template <typename T>
-	static bool ConvertProto3ToUAssetBasicType(TArray<T>& OutBasicTypeArray, const google::protobuf::RepeatedField<T>& InONNXBasicTypeArray);
+	template <typename T1, typename T2>
+	static bool ConvertProto3ToUAssetBasicType(TArray<T1>& OutBasicTypeArray, const google::protobuf::RepeatedField<T2>& InONNXBasicTypeArray);
 
-	template <typename T, typename S>
-	static bool ConvertProto3ToUAssetProtoArrays(TArray<T>& OutProtoArray, const google::protobuf::RepeatedPtrField<S>& InONNXArray);
+	template <typename T1, typename T2>
+	static bool ConvertProto3ToUAssetProtoArrays(TArray<T1>& OutProtoArray, const google::protobuf::RepeatedPtrField<T2>& InONNXArray);
 #endif //WITH_PROTOBUF
 };
 
@@ -79,20 +79,27 @@ private:
 
 #ifdef WITH_PROTOBUF
 
-template <typename T>
-bool FModelProtoConverter::ConvertProto3ToUAssetBasicType(TArray<T>& OutBasicTypeArray, const google::protobuf::RepeatedField<T>& InONNXBasicTypeArray)
+template <typename T1, typename T2>
+bool FModelProtoConverter::ConvertProto3ToUAssetBasicType(TArray<T1>& OutBasicTypeArray, const google::protobuf::RepeatedField<T2>& InONNXBasicTypeArray)
 {
+	if (sizeof(T1) != sizeof(T2))
+	{
+		UE_LOG(LogModelProtoFileReader, Warning, TEXT("FModelProtoConverter::ConvertProto3ToUAssetBasicType(): sizeof(T1) == sizeof(T2) failed (%d != %d)."),
+			sizeof(T1), sizeof(T2));
+		return false;
+	}
+
 	const int32 ArraySize = InONNXBasicTypeArray.size();
 	OutBasicTypeArray.SetNumUninitialized(ArraySize);
 	
-	FMemory::Memcpy(OutBasicTypeArray.GetData(), &InONNXBasicTypeArray.Get(0), sizeof(T)*ArraySize);
+	FMemory::Memcpy(OutBasicTypeArray.GetData(), &InONNXBasicTypeArray.Get(0), sizeof(T1)*ArraySize);
 
 	// Return
 	return true;
 }
 
-template <typename T, typename S>
-static bool FModelProtoConverter::ConvertProto3ToUAssetProtoArrays(TArray<T>& OutProtoArray, const google::protobuf::RepeatedPtrField<S>& InONNXArray)
+template <typename T1, typename T2>
+bool FModelProtoConverter::ConvertProto3ToUAssetProtoArrays(TArray<T1>& OutProtoArray, const google::protobuf::RepeatedPtrField<T2>& InONNXArray)
 {
 	const int32 ArraySize = InONNXArray.size();
 	OutProtoArray.SetNum(ArraySize);
@@ -101,7 +108,7 @@ static bool FModelProtoConverter::ConvertProto3ToUAssetProtoArrays(TArray<T>& Ou
 	{
 		if (!ConvertProto3ToUAsset(OutProtoArray[ArrayIndex], InONNXArray[ArrayIndex]))
 		{
-			UE_LOG(LogModelProtoFileReader, Warning, TEXT("FModelProtoConverter::ConvertFromONNXProto3Ifstream(): ConvertProto3toUAssetProtoArrays() failed."));
+			UE_LOG(LogModelProtoFileReader, Warning, TEXT("FModelProtoConverter::ConvertFromONNXProto3Ifstream(): ConvertProto3ToUAsset() failed."));
 			return false;
 		}
 	}
