NeuralNetworkInference (NNI) plugin:
- ModelProtoConverter VS static warnings fixed.

#rnx

diff --git a/Engine/Plugins/Experimental/NeuralNetworkInference/Source/Private/ModelProtoFileReader/Private/ModelProtoConverter.cpp b/Engine/Plugins/Experimental/NeuralNetworkInference/Source/Private/ModelProtoFileReader/Private/ModelProtoConverter.cpp
--- a/Engine/Plugins/Experimental/NeuralNetworkInference/Source/Private/ModelProtoFileReader/Private/ModelProtoConverter.cpp
+++ b/Engine/Plugins/Experimental/NeuralNetworkInference/Source/Private/ModelProtoFileReader/Private/ModelProtoConverter.cpp
@@ -1,33 +1,72 @@
 // Copyright Epic Games, Inc. All Rights Reserved.
 
 #include "ModelProtoConverter.h"
+#include "ModelProtoFileReaderUtils.h"
 
+// Protobuf includes
+#ifdef WITH_PROTOBUF
+#include "ThirdPartyWarningDisabler.h"
+NNI_THIRD_PARTY_INCLUDES_START
+#undef TEXT
+#undef check
+#include "onnx.proto3.pb.h"
+NNI_THIRD_PARTY_INCLUDES_END
+#endif //WITH_PROTOBUF
 
 
-/* FModelProtoConverter public functions
+
+/* FPrivateModelProtoConverter public functions
 *****************************************************************************/
 
-bool FModelProtoConverter::ConvertFromONNXProto3Ifstream(FModelProto& OutModelProto, std::istream& InIfstream)
-{
 #ifdef WITH_PROTOBUF
-	// Protobuf onnx::ModelProto from Ifstream
-	onnx::ModelProto ModelProto;
-	ModelProto.ParseFromIstream(&InIfstream);
-	return ConvertProto3ToUAsset(OutModelProto, ModelProto);
-#else //WITH_PROTOBUF
-	UE_LOG(LogModelProtoFileReader, Warning, TEXT("FModelProtoConverter::ConvertFromONNXProto3Ifstream(): WITH_PROTOBUF was not defined."));
-	return false
-#endif //WITH_PROTOBUF
-}
 
+class FPrivateModelProtoConverter
+{
+public:
+	static bool ConvertProto3ToUAsset(FModelProto& OutModelProto, const onnx::ModelProto& InONNXModelProto);
 
+	static bool ConvertProto3ToUAsset(FOperatorSetIdProto& OutOperatorSetIdProto, const onnx::OperatorSetIdProto& InONNXOperatorSetIdProto);
 
-/* FModelProtoConverter private functions
-*****************************************************************************/
+	static bool ConvertProto3ToUAsset(FTrainingInfoProto& OutTrainingInfoProto, const onnx::TrainingInfoProto& InONNXTrainingInfoProto);
 
-#ifdef WITH_PROTOBUF
+	static bool ConvertProto3ToUAsset(FGraphProto& OutGraphProto, const onnx::GraphProto& InONNXGraphProto);
+
+	static bool ConvertProto3ToUAsset(FNodeProto& OutNodeProto, const onnx::NodeProto& InONNXNodeProto);
 
-bool FModelProtoConverter::ConvertProto3ToUAsset(FModelProto& OutModelProto, const onnx::ModelProto& InONNXModelProto)
+	static bool ConvertProto3ToUAsset(FTensorAnnotation& OutTensorAnnotation, const onnx::TensorAnnotation& InONNXTensorAnnotation);
+
+	static bool ConvertProto3ToUAsset(FValueInfoProto& OutValueInfoProto, const onnx::ValueInfoProto& InONNXValueInfoProto);
+
+	static bool ConvertProto3ToUAsset(FAttributeProto& OutAttributeProto, const onnx::AttributeProto& InONNXAttributeProto);
+
+	static bool ConvertProto3ToUAsset(FTypeProto& OutTypeProto, const onnx::TypeProto& InONNXTypeProto);
+
+	static bool ConvertProto3ToUAsset(FSparseTensorProto& OutSparseTensorProto, const onnx::SparseTensorProto& InONNXSparseTensorProto);
+
+	static bool ConvertProto3ToUAsset(FTypeProtoTensor& OutTypeProtoTensor, const onnx::TypeProto_Tensor& InONNXTypeProtoTensor);
+
+	static bool ConvertProto3ToUAsset(FTensorProto& OutTensorProto, const onnx::TensorProto& InONNXTensorProto);
+
+	static bool ConvertProto3ToUAsset(FTensorShapeProto& OutTensorShapeProto, const onnx::TensorShapeProto& InONNXTensorShapeProto);
+
+	static bool ConvertProto3ToUAsset(FStringStringEntryProto& OutStringStringEntryProto, const onnx::StringStringEntryProto& FStringStringEntryProto);
+
+	static bool ConvertProto3ToUAsset(FTensorShapeProtoDimension& OutTensorShapeProtoDimension, const onnx::TensorShapeProto_Dimension& InONNXTensorShapeProtoDimension);
+
+	static bool ConvertProto3ToUAsset(FTensorProtoSegment& OutTensorProtoSegment, const onnx::TensorProto_Segment& InONNXTensorProtoSegment);
+
+	static bool ConvertProto3ToUAssetFString(TArray<FString>& OutFStringArray, const google::protobuf::RepeatedPtrField<std::basic_string<char, std::char_traits<char>, std::allocator<char>>>& InONNXFStringArray);
+
+	static bool ConvertProto3ToUAssetUInt8(TArray<uint8>& OutDataArray, const std::string& InRawDataString);
+
+	template <typename T1, typename T2>
+	static bool ConvertProto3ToUAssetBasicType(TArray<T1>& OutBasicTypeArray, const google::protobuf::RepeatedField<T2>& InONNXBasicTypeArray);
+
+	template <typename T1, typename T2>
+	static bool ConvertProto3ToUAssetProtoArrays(TArray<T1>& OutProtoArray, const google::protobuf::RepeatedPtrField<T2>& InONNXArray);
+};
+
+bool FPrivateModelProtoConverter::ConvertProto3ToUAsset(FModelProto& OutModelProto, const onnx::ModelProto& InONNXModelProto)
 {
 	// FModelProto from Protobuf onnx::ModelProto
 	OutModelProto = FModelProto();
@@ -36,7 +75,7 @@ bool FModelProtoConverter::ConvertProto3ToUAsset(FModelProto& OutModelProto, con
 	// OperatorSetIds
 	if (!ConvertProto3ToUAssetProtoArrays(OutModelProto.OpsetImport, InONNXModelProto.opset_import()))
 	{
-		UE_LOG(LogModelProtoFileReader, Warning, TEXT("FModelProtoConverter::ConvertProto3ToUAsset(): ConvertProto3ToUAsset(OperatorSetId) failed."));
+		UE_LOG(LogModelProtoFileReader, Warning, TEXT("FPrivateModelProtoConverter::ConvertProto3ToUAsset(): ConvertProto3ToUAsset(OperatorSetId) failed."));
 		return false;
 	}
 	OutModelProto.ProducerName = UTF8_TO_TCHAR(InONNXModelProto.producer_name().c_str());
@@ -47,19 +86,19 @@ bool FModelProtoConverter::ConvertProto3ToUAsset(FModelProto& OutModelProto, con
 	// Graph
 	if (!ConvertProto3ToUAsset(OutModelProto.Graph, InONNXModelProto.graph()))
 	{
-		UE_LOG(LogModelProtoFileReader, Warning, TEXT("FModelProtoConverter::ConvertProto3ToUAsset(): ConvertProto3GraphToUAsset(Graph) failed."));
+		UE_LOG(LogModelProtoFileReader, Warning, TEXT("FPrivateModelProtoConverter::ConvertProto3ToUAsset(): ConvertProto3GraphToUAsset(Graph) failed."));
 		return false;
 	}
 	// StringStringEntrys
 	if (!ConvertProto3ToUAssetProtoArrays(OutModelProto.MetadataProps, InONNXModelProto.metadata_props()))
 	{
-		UE_LOG(LogModelProtoFileReader, Warning, TEXT("FModelProtoConverter::ConvertProto3ToUAsset(): ConvertProto3ToUAsset(StringStringEntry) failed."));
+		UE_LOG(LogModelProtoFileReader, Warning, TEXT("FPrivateModelProtoConverter::ConvertProto3ToUAsset(): ConvertProto3ToUAsset(StringStringEntry) failed."));
 		return false;
 	}
 	// TrainingInfos
 	if (!ConvertProto3ToUAssetProtoArrays(OutModelProto.TrainingInfo, InONNXModelProto.training_info()))
 	{
-		UE_LOG(LogModelProtoFileReader, Warning, TEXT("FModelProtoConverter::ConvertProto3ToUAsset(): ConvertProto3ToUAsset(TrainingInfo) failed."));
+		UE_LOG(LogModelProtoFileReader, Warning, TEXT("FPrivateModelProtoConverter::ConvertProto3ToUAsset(): ConvertProto3ToUAsset(TrainingInfo) failed."));
 		return false;
 	}
 
@@ -67,7 +106,7 @@ bool FModelProtoConverter::ConvertProto3ToUAsset(FModelProto& OutModelProto, con
 	return true;
 }
 
-bool FModelProtoConverter::ConvertProto3ToUAsset(FOperatorSetIdProto& OutOperatorSetIdProto, const onnx::OperatorSetIdProto& InONNXOperatorSetIdProto)
+bool FPrivateModelProtoConverter::ConvertProto3ToUAsset(FOperatorSetIdProto& OutOperatorSetIdProto, const onnx::OperatorSetIdProto& InONNXOperatorSetIdProto)
 {
 	OutOperatorSetIdProto = FOperatorSetIdProto();
 
@@ -78,7 +117,7 @@ bool FModelProtoConverter::ConvertProto3ToUAsset(FOperatorSetIdProto& OutOperato
 	return true;
 }
 
-bool FModelProtoConverter::ConvertProto3ToUAsset(FTrainingInfoProto& OutTrainingInfoProto, const onnx::TrainingInfoProto& InONNXTrainingInfoProto)
+bool FPrivateModelProtoConverter::ConvertProto3ToUAsset(FTrainingInfoProto& OutTrainingInfoProto, const onnx::TrainingInfoProto& InONNXTrainingInfoProto)
 {
 	OutTrainingInfoProto = FTrainingInfoProto();
 
@@ -86,13 +125,13 @@ bool FModelProtoConverter::ConvertProto3ToUAsset(FTrainingInfoProto& OutTraining
 	// Graph
 	if (!ConvertProto3ToUAsset(OutTrainingInfoProto.Initialization, InONNXTrainingInfoProto.initialization()) || !ConvertProto3ToUAsset(OutTrainingInfoProto.Algorithm, InONNXTrainingInfoProto.algorithm()))
 	{
-		UE_LOG(LogModelProtoFileReader, Warning, TEXT("FModelProtoConverter::ConvertProto3ToUAsset(): ConvertProto3GraphToUAsset(Graph) failed."));
+		UE_LOG(LogModelProtoFileReader, Warning, TEXT("FPrivateModelProtoConverter::ConvertProto3ToUAsset(): ConvertProto3GraphToUAsset(Graph) failed."));
 		return false;
 	}
 	// StringStringEntrys
 	if (!ConvertProto3ToUAssetProtoArrays(OutTrainingInfoProto.InitializationBinding, InONNXTrainingInfoProto.initialization_binding()) || !ConvertProto3ToUAssetProtoArrays(OutTrainingInfoProto.UpdateBinding, InONNXTrainingInfoProto.update_binding()))
 	{
-		UE_LOG(LogModelProtoFileReader, Warning, TEXT("FModelProtoConverter::ConvertProto3ToUAsset(): ConvertProto3ToUAsset(StringStringEntry) failed."));
+		UE_LOG(LogModelProtoFileReader, Warning, TEXT("FPrivateModelProtoConverter::ConvertProto3ToUAsset(): ConvertProto3ToUAsset(StringStringEntry) failed."));
 		return false;
 	}
 
@@ -100,60 +139,60 @@ bool FModelProtoConverter::ConvertProto3ToUAsset(FTrainingInfoProto& OutTraining
 	return true;
 }
 
-bool FModelProtoConverter::ConvertProto3ToUAsset(FGraphProto& OutGraphProto, const onnx::GraphProto& InONNXGraphProto)
+bool FPrivateModelProtoConverter::ConvertProto3ToUAsset(FGraphProto& OutGraphProto, const onnx::GraphProto& InONNXGraphProto)
 {
 	OutGraphProto = FGraphProto();
 	OutGraphProto.bIsLoaded = true;
 	// Nodes
 	if (!ConvertProto3ToUAssetProtoArrays(OutGraphProto.Node, InONNXGraphProto.node()))
 	{
-		UE_LOG(LogModelProtoFileReader, Warning, TEXT("FModelProtoConverter::ConvertProto3ToUAsset(): ConvertProto3ToUAsset(Node) failed."));
+		UE_LOG(LogModelProtoFileReader, Warning, TEXT("FPrivateModelProtoConverter::ConvertProto3ToUAsset(): ConvertProto3ToUAsset(Node) failed."));
 		return false;
 	}
 	OutGraphProto.Name = UTF8_TO_TCHAR(InONNXGraphProto.name().c_str());
 	// TensorProtos
 	if (!ConvertProto3ToUAssetProtoArrays(OutGraphProto.Initializer, InONNXGraphProto.initializer()))
 	{
-		UE_LOG(LogModelProtoFileReader, Warning, TEXT("FModelProtoConverter::ConvertProto3ToUAsset(): ConvertProto3ToUAsset(Tensor) failed."));
+		UE_LOG(LogModelProtoFileReader, Warning, TEXT("FPrivateModelProtoConverter::ConvertProto3ToUAsset(): ConvertProto3ToUAsset(Tensor) failed."));
 		return false;
 	}
 	// SparseTensors
 	if (!ConvertProto3ToUAssetProtoArrays(OutGraphProto.SparseInitializer, InONNXGraphProto.sparse_initializer()))
 	{
-		UE_LOG(LogModelProtoFileReader, Warning, TEXT("FModelProtoConverter::ConvertProto3ToUAsset(): ConvertProto3ToUAsset(SparseTensor) failed."));
+		UE_LOG(LogModelProtoFileReader, Warning, TEXT("FPrivateModelProtoConverter::ConvertProto3ToUAsset(): ConvertProto3ToUAsset(SparseTensor) failed."));
 		return false;
 	}
 	OutGraphProto.DocString = UTF8_TO_TCHAR(InONNXGraphProto.doc_string().c_str());
 	// ValueInfos
 	if (!ConvertProto3ToUAssetProtoArrays(OutGraphProto.Input, InONNXGraphProto.input()) || !ConvertProto3ToUAssetProtoArrays(OutGraphProto.Output, InONNXGraphProto.output()) || !ConvertProto3ToUAssetProtoArrays(OutGraphProto.ValueInfo, InONNXGraphProto.value_info()))
 	{
-		UE_LOG(LogModelProtoFileReader, Warning, TEXT("FModelProtoConverter::ConvertProto3ToUAsset(): ConvertProto3ToUAsset(ValueInfo) failed."));
+		UE_LOG(LogModelProtoFileReader, Warning, TEXT("FPrivateModelProtoConverter::ConvertProto3ToUAsset(): ConvertProto3ToUAsset(ValueInfo) failed."));
 		return false;
 	}
 	// TensorAnnotations
 	if (!ConvertProto3ToUAssetProtoArrays(OutGraphProto.QuantizationAnnotation, InONNXGraphProto.quantization_annotation()))
 	{
-		UE_LOG(LogModelProtoFileReader, Warning, TEXT("FModelProtoConverter::ConvertProto3ToUAsset(): ConvertProto3ToUAsset(TensorAnnotation) failed."));
+		UE_LOG(LogModelProtoFileReader, Warning, TEXT("FPrivateModelProtoConverter::ConvertProto3ToUAsset(): ConvertProto3ToUAsset(TensorAnnotation) failed."));
 		return false;
 	}
 
 	if (false /* If something went wrong at any point, this is a template warning message*/)
 	{
-		UE_LOG(LogModelProtoFileReader, Warning, TEXT("FModelProtoConverter::ConvertProto3ToUAsset(): FGraphProto could not be read from onnx::GraphProto."));
+		UE_LOG(LogModelProtoFileReader, Warning, TEXT("FPrivateModelProtoConverter::ConvertProto3ToUAsset(): FGraphProto could not be read from onnx::GraphProto."));
 		return false;
 	}
 
 	return true;
 }
 
-bool FModelProtoConverter::ConvertProto3ToUAsset(FNodeProto& OutNodeProto, const onnx::NodeProto& InONNXNodeProto)
+bool FPrivateModelProtoConverter::ConvertProto3ToUAsset(FNodeProto& OutNodeProto, const onnx::NodeProto& InONNXNodeProto)
 {
 	OutNodeProto = FNodeProto();
 
 	//FString
 	if (!ConvertProto3ToUAssetFString(OutNodeProto.Input, InONNXNodeProto.input()) || !ConvertProto3ToUAssetFString(OutNodeProto.Output, InONNXNodeProto.output()))
 	{
-		UE_LOG(LogModelProtoFileReader, Warning, TEXT("FModelProtoConverter::ConvertProto3ToUAsset(): ConvertProto3ToUAssetFString() failed."));
+		UE_LOG(LogModelProtoFileReader, Warning, TEXT("FPrivateModelProtoConverter::ConvertProto3ToUAsset(): ConvertProto3ToUAssetFString() failed."));
 		return false;
 	}
 	OutNodeProto.Name = UTF8_TO_TCHAR(InONNXNodeProto.name().c_str());
@@ -162,7 +201,7 @@ bool FModelProtoConverter::ConvertProto3ToUAsset(FNodeProto& OutNodeProto, const
 	// Attributes
 	if (!ConvertProto3ToUAssetProtoArrays(OutNodeProto.Attribute, InONNXNodeProto.attribute()))
 	{
-		UE_LOG(LogModelProtoFileReader, Warning, TEXT("FModelProtoConverter::ConvertProto3ToUAsset(): ConvertProto3ToUAsset(Attribute) failed."));
+		UE_LOG(LogModelProtoFileReader, Warning, TEXT("FPrivateModelProtoConverter::ConvertProto3ToUAsset(): ConvertProto3ToUAsset(Attribute) failed."));
 		return false;
 	}
 	OutNodeProto.DocString = UTF8_TO_TCHAR(InONNXNodeProto.doc_string().c_str());
@@ -171,7 +210,7 @@ bool FModelProtoConverter::ConvertProto3ToUAsset(FNodeProto& OutNodeProto, const
 	return true;
 }
 
-bool FModelProtoConverter::ConvertProto3ToUAsset(FTensorAnnotation& OutTensorAnnotation, const onnx::TensorAnnotation& InONNXTensorAnnotation)
+bool FPrivateModelProtoConverter::ConvertProto3ToUAsset(FTensorAnnotation& OutTensorAnnotation, const onnx::TensorAnnotation& InONNXTensorAnnotation)
 {
 	OutTensorAnnotation = FTensorAnnotation();
 
@@ -179,7 +218,7 @@ bool FModelProtoConverter::ConvertProto3ToUAsset(FTensorAnnotation& OutTensorAnn
 	// StringStringEntrys
 	if (!ConvertProto3ToUAssetProtoArrays(OutTensorAnnotation.QuantParameterTensorNames, InONNXTensorAnnotation.quant_parameter_tensor_names()))
 	{
-		UE_LOG(LogModelProtoFileReader, Warning, TEXT("FModelProtoConverter::ConvertProto3ToUAsset(): ConvertProto3ToUAsset(StringStringEntry) failed."));
+		UE_LOG(LogModelProtoFileReader, Warning, TEXT("FPrivateModelProtoConverter::ConvertProto3ToUAsset(): ConvertProto3ToUAsset(StringStringEntry) failed."));
 		return false;
 	}
 
@@ -187,7 +226,7 @@ bool FModelProtoConverter::ConvertProto3ToUAsset(FTensorAnnotation& OutTensorAnn
 	return true;
 }
 
-bool FModelProtoConverter::ConvertProto3ToUAsset(FValueInfoProto& OutValueInfoProto, const onnx::ValueInfoProto& InONNXValueInfoProto)
+bool FPrivateModelProtoConverter::ConvertProto3ToUAsset(FValueInfoProto& OutValueInfoProto, const onnx::ValueInfoProto& InONNXValueInfoProto)
 {
 	OutValueInfoProto = FValueInfoProto();
 
@@ -195,7 +234,7 @@ bool FModelProtoConverter::ConvertProto3ToUAsset(FValueInfoProto& OutValueInfoPr
 	// Type
 	if (!ConvertProto3ToUAsset(OutValueInfoProto.Type, InONNXValueInfoProto.type()))
 	{
-		UE_LOG(LogModelProtoFileReader, Warning, TEXT("FModelProtoConverter::ConvertProto3ToUAsset(): ConvertProto3ToUAsset(Type) failed."));
+		UE_LOG(LogModelProtoFileReader, Warning, TEXT("FPrivateModelProtoConverter::ConvertProto3ToUAsset(): ConvertProto3ToUAsset(Type) failed."));
 		return false;
 	}
 	OutValueInfoProto.DocString = UTF8_TO_TCHAR(InONNXValueInfoProto.doc_string().c_str());
@@ -204,7 +243,7 @@ bool FModelProtoConverter::ConvertProto3ToUAsset(FValueInfoProto& OutValueInfoPr
 	return true;
 }
 
-bool FModelProtoConverter::ConvertProto3ToUAsset(FAttributeProto& OutAttributeProto, const onnx::AttributeProto& InONNXAttributeProto)
+bool FPrivateModelProtoConverter::ConvertProto3ToUAsset(FAttributeProto& OutAttributeProto, const onnx::AttributeProto& InONNXAttributeProto)
 {
 	OutAttributeProto = FAttributeProto();
 
@@ -217,26 +256,26 @@ bool FModelProtoConverter::ConvertProto3ToUAsset(FAttributeProto& OutAttributePr
 	// TensorProtos
 	if (!ConvertProto3ToUAsset(OutAttributeProto.T, InONNXAttributeProto.t()) || !ConvertProto3ToUAssetProtoArrays(OutAttributeProto.Tensors, InONNXAttributeProto.tensors()))
 	{
-		UE_LOG(LogModelProtoFileReader, Warning, TEXT("FModelProtoConverter::ConvertProto3ToUAsset(): ConvertProto3ToUAsset(Tensor) failed."));
+		UE_LOG(LogModelProtoFileReader, Warning, TEXT("FPrivateModelProtoConverter::ConvertProto3ToUAsset(): ConvertProto3ToUAsset(Tensor) failed."));
 		return false;
 	}
 	//TSharedPtr<FGraphProto> G; /** Used TSharedPtr<> to break the circular dependency */
 	// SparseTensors
 	if (!ConvertProto3ToUAsset(OutAttributeProto.SparseTensor, InONNXAttributeProto.sparse_tensor()) || !ConvertProto3ToUAssetProtoArrays(OutAttributeProto.SparseTensors, InONNXAttributeProto.sparse_tensors()))
 	{
-		UE_LOG(LogModelProtoFileReader, Warning, TEXT("FModelProtoConverter::ConvertProto3ToUAsset(): ConvertProto3ToUAsset(SparseTensor) failed."));
+		UE_LOG(LogModelProtoFileReader, Warning, TEXT("FPrivateModelProtoConverter::ConvertProto3ToUAsset(): ConvertProto3ToUAsset(SparseTensor) failed."));
 		return false;
 	}
 	// floats and int64s
 	if (!ConvertProto3ToUAssetBasicType(OutAttributeProto.Floats, InONNXAttributeProto.floats()) || !ConvertProto3ToUAssetBasicType(OutAttributeProto.Integers, InONNXAttributeProto.ints()))
 	{
-		UE_LOG(LogModelProtoFileReader, Warning, TEXT("FModelProtoConverter::ConvertProto3ToUAsset(): ConvertProto3ToUAssetBasicType() failed."));
+		UE_LOG(LogModelProtoFileReader, Warning, TEXT("FPrivateModelProtoConverter::ConvertProto3ToUAsset(): ConvertProto3ToUAssetBasicType() failed."));
 		return false;
 	}
 	//TArray<FString>
 	if (!ConvertProto3ToUAssetFString(OutAttributeProto.Strings, InONNXAttributeProto.strings()))
 	{
-		UE_LOG(LogModelProtoFileReader, Warning, TEXT("FModelProtoConverter::ConvertProto3ToUAsset(): ConvertProto3ToUAssetFString() failed."));
+		UE_LOG(LogModelProtoFileReader, Warning, TEXT("FPrivateModelProtoConverter::ConvertProto3ToUAsset(): ConvertProto3ToUAssetFString() failed."));
 		return false;
 	}
 	//TArray<TSharedPtr<FGraphProto>> Graphs; /** Used TSharedPtr<> to break the circular dependency */
@@ -245,13 +284,13 @@ bool FModelProtoConverter::ConvertProto3ToUAsset(FAttributeProto& OutAttributePr
 	return true;
 }
 
-bool FModelProtoConverter::ConvertProto3ToUAsset(FTypeProto& OutTypeProto, const onnx::TypeProto& InONNXTypeProto)
+bool FPrivateModelProtoConverter::ConvertProto3ToUAsset(FTypeProto& OutTypeProto, const onnx::TypeProto& InONNXTypeProto)
 {
 	OutTypeProto = FTypeProto();
 	//TypeProto
 	if (!ConvertProto3ToUAsset(OutTypeProto.TensorType, InONNXTypeProto.tensor_type()))
 	{
-		UE_LOG(LogModelProtoFileReader, Warning, TEXT("FModelProtoConverter::ConvertProto3ToUAsset(): ConvertProto3ToUAsset(TypeTensor) failed."));
+		UE_LOG(LogModelProtoFileReader, Warning, TEXT("FPrivateModelProtoConverter::ConvertProto3ToUAsset(): ConvertProto3ToUAsset(TypeTensor) failed."));
 		return false;
 	}
 	OutTypeProto.Denotation = UTF8_TO_TCHAR(InONNXTypeProto.denotation().c_str());
@@ -260,19 +299,19 @@ bool FModelProtoConverter::ConvertProto3ToUAsset(FTypeProto& OutTypeProto, const
 	return true;
 }
 
-bool FModelProtoConverter::ConvertProto3ToUAsset(FSparseTensorProto& OutSparseTensorProto, const onnx::SparseTensorProto& InONNXSparseTensorProto)
+bool FPrivateModelProtoConverter::ConvertProto3ToUAsset(FSparseTensorProto& OutSparseTensorProto, const onnx::SparseTensorProto& InONNXSparseTensorProto)
 {
 	OutSparseTensorProto = FSparseTensorProto();
 	//Tensor
 	if (!ConvertProto3ToUAsset(OutSparseTensorProto.Values, InONNXSparseTensorProto.values()) || !ConvertProto3ToUAsset(OutSparseTensorProto.Indices, InONNXSparseTensorProto.indices()))
 	{
-		UE_LOG(LogModelProtoFileReader, Warning, TEXT("FModelProtoConverter::ConvertProto3ToUAsset(): ConvertProto3ToUAsset(Tensor) failed."));
+		UE_LOG(LogModelProtoFileReader, Warning, TEXT("FPrivateModelProtoConverter::ConvertProto3ToUAsset(): ConvertProto3ToUAsset(Tensor) failed."));
 		return false;
 	}
 	//TArray<int64> Dimensions
 	if (!ConvertProto3ToUAssetBasicType(OutSparseTensorProto.Dimensions, InONNXSparseTensorProto.dims()))
 	{
-		UE_LOG(LogModelProtoFileReader, Warning, TEXT("FModelProtoConverter::ConvertProto3ToUAsset(): ConvertProto3ToUAssetBasicType() failed."));
+		UE_LOG(LogModelProtoFileReader, Warning, TEXT("FPrivateModelProtoConverter::ConvertProto3ToUAsset(): ConvertProto3ToUAssetBasicType() failed."));
 		return false;
 	}
 
@@ -280,7 +319,7 @@ bool FModelProtoConverter::ConvertProto3ToUAsset(FSparseTensorProto& OutSparseTe
 	return true;
 }
 
-bool FModelProtoConverter::ConvertProto3ToUAsset(FTypeProtoTensor& OutTypeProtoTensor, const onnx::TypeProto_Tensor& InONNXTypeProtoTensor)
+bool FPrivateModelProtoConverter::ConvertProto3ToUAsset(FTypeProtoTensor& OutTypeProtoTensor, const onnx::TypeProto_Tensor& InONNXTypeProtoTensor)
 {
 	OutTypeProtoTensor = FTypeProtoTensor();
 	
@@ -288,7 +327,7 @@ bool FModelProtoConverter::ConvertProto3ToUAsset(FTypeProtoTensor& OutTypeProtoT
 	// TensorShape
 	if (!ConvertProto3ToUAsset(OutTypeProtoTensor.Shape, InONNXTypeProtoTensor.shape()))
 	{
-		UE_LOG(LogModelProtoFileReader, Warning, TEXT("FModelProtoConverter::ConvertProto3ToUAsset(): ConvertProto3ToUAsset(TensorShape) failed."));
+		UE_LOG(LogModelProtoFileReader, Warning, TEXT("FPrivateModelProtoConverter::ConvertProto3ToUAsset(): ConvertProto3ToUAsset(TensorShape) failed."));
 		return false;
 	}
 
@@ -296,14 +335,14 @@ bool FModelProtoConverter::ConvertProto3ToUAsset(FTypeProtoTensor& OutTypeProtoT
 	return true;
 }
 
-bool FModelProtoConverter::ConvertProto3ToUAsset(FTensorProto& OutTensorProto, const onnx::TensorProto& InONNXTensorProto)
+bool FPrivateModelProtoConverter::ConvertProto3ToUAsset(FTensorProto& OutTensorProto, const onnx::TensorProto& InONNXTensorProto)
 {
 	OutTensorProto = FTensorProto();
 
 	// TArray<uint8>
 	if (!ConvertProto3ToUAssetUInt8(OutTensorProto.RawData, InONNXTensorProto.raw_data()))
 	{
-		UE_LOG(LogModelProtoFileReader, Warning, TEXT("FModelProtoConverter::ConvertProto3ToUAsset(): ConvertProto3ToUAssetUInt8() failed."));
+		UE_LOG(LogModelProtoFileReader, Warning, TEXT("FPrivateModelProtoConverter::ConvertProto3ToUAsset(): ConvertProto3ToUAssetUInt8() failed."));
 		return false;
 	}
 	//TArray<int64, float, int32, double, uint64>
@@ -311,20 +350,20 @@ bool FModelProtoConverter::ConvertProto3ToUAsset(FTensorProto& OutTensorProto, c
 		!ConvertProto3ToUAssetBasicType(OutTensorProto.Int32Data, InONNXTensorProto.int32_data()) || !ConvertProto3ToUAssetBasicType(OutTensorProto.Int64Data, InONNXTensorProto.int64_data()) ||
 		!ConvertProto3ToUAssetBasicType(OutTensorProto.DoubleData, InONNXTensorProto.double_data()) || !ConvertProto3ToUAssetBasicType(OutTensorProto.UInt64Data, InONNXTensorProto.uint64_data()))
 	{
-		UE_LOG(LogModelProtoFileReader, Warning, TEXT("FModelProtoConverter::ConvertProto3ToUAsset(): ConvertProto3ToUAssetBasicType() failed."));
+		UE_LOG(LogModelProtoFileReader, Warning, TEXT("FPrivateModelProtoConverter::ConvertProto3ToUAsset(): ConvertProto3ToUAssetBasicType() failed."));
 		return false;
 	}
 	OutTensorProto.DataType = (ETensorProtoDataType)InONNXTensorProto.data_type();
 	//FTensorProtoSegment
 	if (!ConvertProto3ToUAsset(OutTensorProto.Segment, InONNXTensorProto.segment()))
 	{
-		UE_LOG(LogModelProtoFileReader, Warning, TEXT("FModelProtoConverter::ConvertProto3ToUAsset(): ConvertProto3ToUAsset(TensorProtoSegment) failed."));
+		UE_LOG(LogModelProtoFileReader, Warning, TEXT("FPrivateModelProtoConverter::ConvertProto3ToUAsset(): ConvertProto3ToUAsset(TensorProtoSegment) failed."));
 		return false;
 	}
 	//TArray<FString>
 	if (!ConvertProto3ToUAssetFString(OutTensorProto.StringData, InONNXTensorProto.string_data()))
 	{
-		UE_LOG(LogModelProtoFileReader, Warning, TEXT("FModelProtoConverter::ConvertProto3ToUAsset(): ConvertProto3ToUAssetFString() failed."));
+		UE_LOG(LogModelProtoFileReader, Warning, TEXT("FPrivateModelProtoConverter::ConvertProto3ToUAsset(): ConvertProto3ToUAssetFString() failed."));
 		return false;
 	}
 	OutTensorProto.Name = UTF8_TO_TCHAR(InONNXTensorProto.name().c_str());
@@ -332,7 +371,7 @@ bool FModelProtoConverter::ConvertProto3ToUAsset(FTensorProto& OutTensorProto, c
 	// StringStringEntryProtos
 	if (!ConvertProto3ToUAssetProtoArrays(OutTensorProto.ExternalData, InONNXTensorProto.external_data()))
 	{
-		UE_LOG(LogModelProtoFileReader, Warning, TEXT("FModelProtoConverter::ConvertProto3ToUAsset(): ConvertProto3ToUAsset(StringStringEntry) failed."));
+		UE_LOG(LogModelProtoFileReader, Warning, TEXT("FPrivateModelProtoConverter::ConvertProto3ToUAsset(): ConvertProto3ToUAsset(StringStringEntry) failed."));
 		return false;
 	}
 	OutTensorProto.DataLocation = (ETensorProtoDataLocation)InONNXTensorProto.data_location();
@@ -341,13 +380,13 @@ bool FModelProtoConverter::ConvertProto3ToUAsset(FTensorProto& OutTensorProto, c
 	return true;
 }
 
-bool FModelProtoConverter::ConvertProto3ToUAsset(FTensorShapeProto& OutTensorShapeProto, const onnx::TensorShapeProto& InONNXTensorShapeProto)
+bool FPrivateModelProtoConverter::ConvertProto3ToUAsset(FTensorShapeProto& OutTensorShapeProto, const onnx::TensorShapeProto& InONNXTensorShapeProto)
 {
 	OutTensorShapeProto = FTensorShapeProto();
 	// TensorShapeDimensions
 	if (!ConvertProto3ToUAssetProtoArrays(OutTensorShapeProto.Dim, InONNXTensorShapeProto.dim()))
 	{
-		UE_LOG(LogModelProtoFileReader, Warning, TEXT("FModelProtoConverter::ConvertProto3ToUAsset(): ConvertProto3ToUAsset(TensorShapeProtoDimension) failed."));
+		UE_LOG(LogModelProtoFileReader, Warning, TEXT("FPrivateModelProtoConverter::ConvertProto3ToUAsset(): ConvertProto3ToUAsset(TensorShapeProtoDimension) failed."));
 		return false;
 	}
 
@@ -355,7 +394,7 @@ bool FModelProtoConverter::ConvertProto3ToUAsset(FTensorShapeProto& OutTensorSha
 	return true;
 }
 
-bool FModelProtoConverter::ConvertProto3ToUAsset(FStringStringEntryProto& OutStringStringEntryProto, const onnx::StringStringEntryProto& InONNXStringStringEntryProto)
+bool FPrivateModelProtoConverter::ConvertProto3ToUAsset(FStringStringEntryProto& OutStringStringEntryProto, const onnx::StringStringEntryProto& InONNXStringStringEntryProto)
 {
 	OutStringStringEntryProto = FStringStringEntryProto();
 
@@ -366,7 +405,7 @@ bool FModelProtoConverter::ConvertProto3ToUAsset(FStringStringEntryProto& OutStr
 	return true;
 }
 
-bool FModelProtoConverter::ConvertProto3ToUAsset(FTensorShapeProtoDimension& OutTensorShapeProtoDimension, const onnx::TensorShapeProto_Dimension& InONNXTensorShapeProtoDimension)
+bool FPrivateModelProtoConverter::ConvertProto3ToUAsset(FTensorShapeProtoDimension& OutTensorShapeProtoDimension, const onnx::TensorShapeProto_Dimension& InONNXTensorShapeProtoDimension)
 {
 	OutTensorShapeProtoDimension = FTensorShapeProtoDimension();
 
@@ -377,7 +416,7 @@ bool FModelProtoConverter::ConvertProto3ToUAsset(FTensorShapeProtoDimension& Out
 	return true;
 }
 
-bool FModelProtoConverter::ConvertProto3ToUAsset(FTensorProtoSegment& OutTensorProtoSegment, const onnx::TensorProto_Segment& InONNXTensorProtoSegment)
+bool FPrivateModelProtoConverter::ConvertProto3ToUAsset(FTensorProtoSegment& OutTensorProtoSegment, const onnx::TensorProto_Segment& InONNXTensorProtoSegment)
 {
 	OutTensorProtoSegment = FTensorProtoSegment();
 
@@ -388,7 +427,7 @@ bool FModelProtoConverter::ConvertProto3ToUAsset(FTensorProtoSegment& OutTensorP
 	return true;
 }
 
-bool FModelProtoConverter::ConvertProto3ToUAssetFString(TArray<FString>& OutFStringArray, const google::protobuf::RepeatedPtrField<std::basic_string<char, std::char_traits<char>, std::allocator<char>>>& InONNXFStringArray)
+bool FPrivateModelProtoConverter::ConvertProto3ToUAssetFString(TArray<FString>& OutFStringArray, const google::protobuf::RepeatedPtrField<std::basic_string<char, std::char_traits<char>, std::allocator<char>>>& InONNXFStringArray)
 {
 	const int32 ArraySize = InONNXFStringArray.size();
 	OutFStringArray.SetNum(ArraySize);
@@ -400,7 +439,7 @@ bool FModelProtoConverter::ConvertProto3ToUAssetFString(TArray<FString>& OutFStr
 	return true;
 }
 
-bool FModelProtoConverter::ConvertProto3ToUAssetUInt8(TArray<uint8>& OutDataArray, const std::string& InRawDataString)
+bool FPrivateModelProtoConverter::ConvertProto3ToUAssetUInt8(TArray<uint8>& OutDataArray, const std::string& InRawDataString)
 {
 	const int32 ArraySize = InRawDataString.size();
 	OutDataArray.SetNumUninitialized(ArraySize);
@@ -411,4 +450,58 @@ bool FModelProtoConverter::ConvertProto3ToUAssetUInt8(TArray<uint8>& OutDataArra
 	return true;
 }
 
+template <typename T1, typename T2>
+bool FPrivateModelProtoConverter::ConvertProto3ToUAssetBasicType(TArray<T1>& OutBasicTypeArray, const google::protobuf::RepeatedField<T2>& InONNXBasicTypeArray)
+{
+	if (sizeof(T1) != sizeof(T2))
+	{
+		UE_LOG(LogModelProtoFileReader, Warning, TEXT("FPrivateModelProtoConverter::ConvertProto3ToUAssetBasicType(): sizeof(T1) == sizeof(T2) failed (%d != %d)."),
+			sizeof(T1), sizeof(T2));
+		return false;
+	}
+
+	const int32 ArraySize = InONNXBasicTypeArray.size();
+	OutBasicTypeArray.SetNumUninitialized(ArraySize);
+
+	FMemory::Memcpy(OutBasicTypeArray.GetData(), &InONNXBasicTypeArray.Get(0), sizeof(T1) * ArraySize);
+
+	// Return
+	return true;
+}
+
+template <typename T1, typename T2>
+bool FPrivateModelProtoConverter::ConvertProto3ToUAssetProtoArrays(TArray<T1>& OutProtoArray, const google::protobuf::RepeatedPtrField<T2>& InONNXArray)
+{
+	const int32 ArraySize = InONNXArray.size();
+	OutProtoArray.SetNum(ArraySize);
+
+	for (int32 ArrayIndex = 0; ArrayIndex < ArraySize; ArrayIndex++)
+	{
+		if (!ConvertProto3ToUAsset(OutProtoArray[ArrayIndex], InONNXArray[ArrayIndex]))
+		{
+			UE_LOG(LogModelProtoFileReader, Warning, TEXT("FPrivateModelProtoConverter::ConvertFromONNXProto3Ifstream(): ConvertProto3ToUAsset() failed."));
+			return false;
+		}
+	}
+	return true;
+}
+
 #endif //WITH_PROTOBUF
+
+
+
+/* FModelProtoConverter public functions
+*****************************************************************************/
+
+bool FModelProtoConverter::ConvertFromONNXProto3Ifstream(FModelProto& OutModelProto, std::istream& InIfstream)
+{
+#ifdef WITH_PROTOBUF
+	// Protobuf onnx::ModelProto from Ifstream
+	onnx::ModelProto ModelProto;
+	ModelProto.ParseFromIstream(&InIfstream);
+	return FPrivateModelProtoConverter::ConvertProto3ToUAsset(OutModelProto, ModelProto);
+#else //WITH_PROTOBUF
+	UE_LOG(LogModelProtoFileReader, Warning, TEXT("FModelProtoConverter::ConvertFromONNXProto3Ifstream(): WITH_PROTOBUF was not defined."));
+	return false
+#endif //WITH_PROTOBUF
+}
diff --git a/Engine/Plugins/Experimental/NeuralNetworkInference/Source/Private/ModelProtoFileReader/Private/ModelProtoConverter.h b/Engine/Plugins/Experimental/NeuralNetworkInference/Source/Private/ModelProtoFileReader/Private/ModelProtoConverter.h
--- a/Engine/Plugins/Experimental/NeuralNetworkInference/Source/Private/ModelProtoFileReader/Private/ModelProtoConverter.h
+++ b/Engine/Plugins/Experimental/NeuralNetworkInference/Source/Private/ModelProtoFileReader/Private/ModelProtoConverter.h
@@ -4,17 +4,7 @@
 
 #include "CoreMinimal.h"
 #include "ModelProto.h"
-#include "ModelProtoFileReaderUtils.h"
-
-// Protobuf includes
-#ifdef WITH_PROTOBUF
-#include "ThirdPartyWarningDisabler.h"
-NNI_THIRD_PARTY_INCLUDES_START
-#undef TEXT
-#undef check
-#include "onnx.proto3.pb.h"
-NNI_THIRD_PARTY_INCLUDES_END
-#endif //WITH_PROTOBUF
+#include <istream>
 
 class FModelProtoConverter
 {
@@ -24,95 +14,4 @@ public:
 	 * @return False if it could not initialize it, true if successful.
 	 */
 	static bool ConvertFromONNXProto3Ifstream(FModelProto& OutModelProto, std::istream& InIfstream);
-
-#ifdef WITH_PROTOBUF
-private:
-	static bool ConvertProto3ToUAsset(FModelProto& OutModelProto, const onnx::ModelProto& InONNXModelProto);
-
-	static bool ConvertProto3ToUAsset(FOperatorSetIdProto& OutOperatorSetIdProto, const onnx::OperatorSetIdProto& InONNXOperatorSetIdProto);
-
-	static bool ConvertProto3ToUAsset(FTrainingInfoProto& OutTrainingInfoProto, const onnx::TrainingInfoProto& InONNXTrainingInfoProto);
-
-	static bool ConvertProto3ToUAsset(FGraphProto& OutGraphProto, const onnx::GraphProto& InONNXGraphProto);
-
-	static bool ConvertProto3ToUAsset(FNodeProto& OutNodeProto, const onnx::NodeProto& InONNXNodeProto);
-
-	static bool ConvertProto3ToUAsset(FTensorAnnotation& OutTensorAnnotation, const onnx::TensorAnnotation& InONNXTensorAnnotation);
-
-	static bool ConvertProto3ToUAsset(FValueInfoProto& OutValueInfoProto, const onnx::ValueInfoProto& InONNXValueInfoProto);
-
-	static bool ConvertProto3ToUAsset(FAttributeProto& OutAttributeProto, const onnx::AttributeProto& InONNXAttributeProto);
-
-	static bool ConvertProto3ToUAsset(FTypeProto& OutTypeProto, const onnx::TypeProto& InONNXTypeProto);
-
-	static bool ConvertProto3ToUAsset(FSparseTensorProto& OutSparseTensorProto, const onnx::SparseTensorProto& InONNXSparseTensorProto);
-
-	static bool ConvertProto3ToUAsset(FTypeProtoTensor& OutTypeProtoTensor, const onnx::TypeProto_Tensor& InONNXTypeProtoTensor);
-
-	static bool ConvertProto3ToUAsset(FTensorProto& OutTensorProto, const onnx::TensorProto& InONNXTensorProto);
-
-	static bool ConvertProto3ToUAsset(FTensorShapeProto& OutTensorShapeProto, const onnx::TensorShapeProto& InONNXTensorShapeProto);
-
-	static bool ConvertProto3ToUAsset(FStringStringEntryProto& OutStringStringEntryProto, const onnx::StringStringEntryProto& FStringStringEntryProto);
-
-	static bool ConvertProto3ToUAsset(FTensorShapeProtoDimension& OutTensorShapeProtoDimension, const onnx::TensorShapeProto_Dimension& InONNXTensorShapeProtoDimension);
-
-	static bool ConvertProto3ToUAsset(FTensorProtoSegment& OutTensorProtoSegment, const onnx::TensorProto_Segment& InONNXTensorProtoSegment);
-
-	static bool ConvertProto3ToUAssetFString(TArray<FString>& OutFStringArray, const google::protobuf::RepeatedPtrField<std::basic_string<char, std::char_traits<char>, std::allocator<char>>>& InONNXFStringArray);
-
-	static bool ConvertProto3ToUAssetUInt8(TArray<uint8>& OutDataArray, const std::string& InRawDataString);
-
-	template <typename T1, typename T2>
-	static bool ConvertProto3ToUAssetBasicType(TArray<T1>& OutBasicTypeArray, const google::protobuf::RepeatedField<T2>& InONNXBasicTypeArray);
-
-	template <typename T1, typename T2>
-	static bool ConvertProto3ToUAssetProtoArrays(TArray<T1>& OutProtoArray, const google::protobuf::RepeatedPtrField<T2>& InONNXArray);
-#endif //WITH_PROTOBUF
 };
-
-
-
-
-/* ModelProtoConverter template functions
- *****************************************************************************/
-
-#ifdef WITH_PROTOBUF
-
-template <typename T1, typename T2>
-bool FModelProtoConverter::ConvertProto3ToUAssetBasicType(TArray<T1>& OutBasicTypeArray, const google::protobuf::RepeatedField<T2>& InONNXBasicTypeArray)
-{
-	if (sizeof(T1) != sizeof(T2))
-	{
-		UE_LOG(LogModelProtoFileReader, Warning, TEXT("FModelProtoConverter::ConvertProto3ToUAssetBasicType(): sizeof(T1) == sizeof(T2) failed (%d != %d)."),
-			sizeof(T1), sizeof(T2));
-		return false;
-	}
-
-	const int32 ArraySize = InONNXBasicTypeArray.size();
-	OutBasicTypeArray.SetNumUninitialized(ArraySize);
-	
-	FMemory::Memcpy(OutBasicTypeArray.GetData(), &InONNXBasicTypeArray.Get(0), sizeof(T1)*ArraySize);
-
-	// Return
-	return true;
-}
-
-template <typename T1, typename T2>
-bool FModelProtoConverter::ConvertProto3ToUAssetProtoArrays(TArray<T1>& OutProtoArray, const google::protobuf::RepeatedPtrField<T2>& InONNXArray)
-{
-	const int32 ArraySize = InONNXArray.size();
-	OutProtoArray.SetNum(ArraySize);
-
-	for (int32 ArrayIndex = 0; ArrayIndex < ArraySize; ArrayIndex++)
-	{
-		if (!ConvertProto3ToUAsset(OutProtoArray[ArrayIndex], InONNXArray[ArrayIndex]))
-		{
-			UE_LOG(LogModelProtoFileReader, Warning, TEXT("FModelProtoConverter::ConvertFromONNXProto3Ifstream(): ConvertProto3ToUAsset() failed."));
-			return false;
-		}
-	}
-	return true;
-}
-
-#endif //WITH_PROTOBUF
