NeuralNetworkInference (NNI) plugin:
- NNILegacy merged into NNI step 3: NNI ORT code moved to *.imp file.

#rnx

diff --git a/Engine/Plugins/Experimental/NeuralNetworkInference/Source/NeuralNetworkInference/Private/NeuralNetwork.cpp b/Engine/Plugins/Experimental/NeuralNetworkInference/Source/NeuralNetworkInference/Private/NeuralNetwork.cpp
--- a/Engine/Plugins/Experimental/NeuralNetworkInference/Source/NeuralNetworkInference/Private/NeuralNetwork.cpp
+++ b/Engine/Plugins/Experimental/NeuralNetworkInference/Source/NeuralNetworkInference/Private/NeuralNetwork.cpp
@@ -4,7 +4,6 @@
 #include "NeuralNetworkInferenceUtils.h"
 #include "EditorFramework/AssetImportData.h"
 #include "HAL/FileManager.h"
-#include "Misc/FileHelper.h"
 #include "Misc/Paths.h"
 
 // Other files with UNeuralNetwork implementation (to be included only once and after the includes above)
@@ -45,143 +44,78 @@ bool UNeuralNetwork::Load(const FString& InModelFilePath)
 	// Clean previous networks
 	bIsLoaded = false;
 
-#ifdef WITH_FULL_NNI_SUPPORT
-	if (InModelFilePath.IsEmpty())
+	// Fill ModelFullFilePath & ModelReadFromDiskInBytes
 	{
-		UE_LOG(LogNeuralNetworkInference, Warning, TEXT("UNeuralNetwork::Load(): Input model path was empty."));
-		return false;
+		// Sanity check
+		if (InModelFilePath.IsEmpty())
+		{
+			UE_LOG(LogNeuralNetworkInference, Warning, TEXT("UNeuralNetwork::Load(): Input model path was empty."));
+			return false;
+		}
+		// Fill ModelFullFilePath
+		ModelFullFilePath = IFileManager::Get().ConvertToAbsolutePathForExternalAppForRead(*InModelFilePath);
+		// Sanity check
+		if (!FPaths::FileExists(ModelFullFilePath))
+		{
+			UE_LOG(LogNeuralNetworkInference, Warning, TEXT("UNeuralNetwork::Load(): Model not found \"%s\"."), *ModelFullFilePath);
+			return false;
+		}
+		// Read file into ModelReadFromDiskInBytes
+		// Source: [REDACTED_PROFILE_URL]/onnxruntime/blob/894fc828587c919d815918c4da6cde314e5d54ed/onnxruntime/test/shared_lib/test_model_loading.cc#L21-L31
+		if (!FFileHelper::LoadFileToArray(ModelReadFromDiskInBytes, *ModelFullFilePath))
+		{
+			UE_LOG(LogNeuralNetworkInference, Warning, TEXT("UNeuralNetwork::Load(): Error reading model \"%s\"."), *ModelFullFilePath);
+			return false;
+		}
 	}
 
-	ModelFullFilePath = IFileManager::Get().ConvertToAbsolutePathForExternalAppForRead(*InModelFilePath);
-
-	// Sanity check
-	if (!FPaths::FileExists(ModelFullFilePath))
+	// UEAndORT
+	if (BackEnd == ENeuralBackEnd::UEAndORT)
 	{
-		UE_LOG(LogNeuralNetworkInference, Warning, TEXT("UNeuralNetwork::Load(): Model not found \"%s\"."), *ModelFullFilePath);
-		return false;
+		return Load();
 	}
-
-	const FRedirectCoutAndCerrToUeLog RedirectCoutAndCerrToUeLog;
-
-	// Initialize and configure ImplBackEndUEAndORT
-	if (!UNeuralNetwork::FImplBackEndUEAndORT::InitializedAndConfigureMembers(ImplBackEndUEAndORT, ModelFullFilePath, GetDeviceType()))
+	// UEOnly
+	else if (BackEnd == ENeuralBackEnd::UEOnly)
 	{
-		UE_LOG(LogNeuralNetworkInference, Warning, TEXT("UNeuralNetwork::Load(): InitializedAndConfigureMembers failed."));
-		return false;
+		UE_LOG(LogNeuralNetworkInference, Warning, TEXT("UNeuralNetwork::Load(): Not implemented for BackEnd = %d yet."), (int32)BackEnd);
 	}
-
-	// Read file into ModelReadFromDiskInBytes
-	// Source: [REDACTED_PROFILE_URL]/onnxruntime/blob/894fc828587c919d815918c4da6cde314e5d54ed/onnxruntime/test/shared_lib/test_model_loading.cc#L21-L31
-	if (!FFileHelper::LoadFileToArray(ModelReadFromDiskInBytes, *ModelFullFilePath))
+	// Unknown
+	else
 	{
-		UE_LOG(LogNeuralNetworkInference, Warning, TEXT("UNeuralNetwork::Load(): Error reading model \"%s\"."), *ModelFullFilePath);
-		return false;
+		UE_LOG(LogNeuralNetworkInference, Warning, TEXT("UNeuralNetwork::Load(): Unknown BackEnd = %d."), (int32)BackEnd);
 	}
-	return Load();
-
-
-
-//	// Note: This code uses the changes in the ONNX Runtime API, which are not needed for desktop platforms (where ONNX/ProtoBuf is supported)
-//	// Fill ModelReadFromDiskInBytes from ONNX/ORT file
-//	const FString FileExtension = FPaths::GetExtension(ModelFullFilePath, /*bIncludeDot*/ false);
-//	const char* const FilePathCharPtr = TCHAR_TO_ANSI(*ModelFullFilePath);
-//	// If ONNX file, turn into ORT format first
-//	if (FileExtension.Equals(TEXT("onnx"), ESearchCase::IgnoreCase))
-//	{
-//		FString ONNXPathPart, ONNXFilenamePart, ONNXExtensionPart;
-//		FPaths::Split(ModelFullFilePath, ONNXPathPart, ONNXFilenamePart, ONNXExtensionPart);
-//		const FString OutputORTOptimizedModelPath = ONNXPathPart / ONNXFilenamePart + TEXT(".ort");
-//#ifdef _WIN32
-//		ImplBackEndUEAndORT->SessionOptions->SetOptimizedModelFilePath(*OutputORTOptimizedModelPath);
-//#else
-//		ImplBackEndUEAndORT->SessionOptions->SetOptimizedModelFilePath(TCHAR_TO_ANSI(*OutputORTOptimizedModelPath));
-//#endif //_WIN32
-//		// ONNX --> ORT conversion
-//		// This session is just temporarily opened so the ORT model file can be generated
-//		ImplBackEndUEAndORT->Session = MakeUnique<Ort::Session>(*ImplBackEndUEAndORT->Environment,
-//#ifdef _WIN32
-//			*ModelFullFilePath,
-//#else
-//			FilePathCharPtr,
-//#endif
-//			*ImplBackEndUEAndORT->SessionOptions);
-//
-//		// Read model from OutputORTOptimizedModelPath
-//		return Load(OutputORTOptimizedModelPath);
-//	}
-//	// Create session (it should be ORT file at this point), and read ModelReadFromDiskInBytes if not empty
-//	else if (FileExtension.Equals(TEXT("ort"), ESearchCase::IgnoreCase))
-//	{
-//		// Read model from ModelFullFilePath
-//		std::vector<uint8_t> OutputModelReadFromDiskInBytesVector;
-//		ImplBackEndUEAndORT->Session = MakeUnique<Ort::Session>(*ImplBackEndUEAndORT->Environment,
-//#ifdef _WIN32
-//			*ModelFullFilePath,
-//#else
-//			FilePathCharPtr,
-//#endif
-//			*ImplBackEndUEAndORT->SessionOptions, &OutputModelReadFromDiskInBytesVector);
-//
-//		// Fill ModelReadFromDiskInBytes
-//		const int32 ArraySize = OutputModelReadFromDiskInBytesVector.size();
-//		if (ArraySize > 0)
-//		{
-//			ModelReadFromDiskInBytes.SetNumUninitialized(ArraySize);
-//			FMemory::Memcpy(ModelReadFromDiskInBytes.GetData(), &OutputModelReadFromDiskInBytesVector[0], ArraySize);
-//		}
-//
-//		return Load();
-//	}
-//	// Else
-//	UE_LOG(LogNeuralNetworkInference, Warning, TEXT("UNeuralNetwork::Load(): Unknown file format \"%s\" used."), *FileExtension);
-//	return false;
 
-#else
-	UE_LOG(LogNeuralNetworkInference, Warning, TEXT("UNeuralNetwork::Load(): Platform or Operating System not suported yet."));
-	return false;
-#endif //WITH_FULL_NNI_SUPPORT
+	return bIsLoaded;
 }
 #endif //WITH_EDITOR
 
 
 bool UNeuralNetwork::Load()
 {
-#ifdef WITH_FULL_NNI_SUPPORT
 	// Clean previous networks
 	bIsLoaded = false;
 
-	const FRedirectCoutAndCerrToUeLog RedirectCoutAndCerrToUeLog;
-
-	// Initialize and configure ImplBackEndUEAndORT
-	if (!UNeuralNetwork::FImplBackEndUEAndORT::InitializedAndConfigureMembers(ImplBackEndUEAndORT, ModelFullFilePath, GetDeviceType()))
+	// UEAndORT
+	if (BackEnd == ENeuralBackEnd::UEAndORT)
 	{
-		UE_LOG(LogNeuralNetworkInference, Warning, TEXT("UNeuralNetwork::Load(): InitializedAndConfigureMembers failed."));
-		return false;
+#ifdef WITH_FULL_NNI_SUPPORT
+		bIsLoaded = UNeuralNetwork::FImplBackEndUEAndORT::Load(ImplBackEndUEAndORT, InputTensors, OutputTensors, AreInputTensorSizesVariable, ModelReadFromDiskInBytes, ModelFullFilePath, GetDeviceType());
+#else
+		UE_LOG(LogNeuralNetworkInference, Warning, TEXT("UNeuralNetwork::Load(): Platform or Operating System not suported yet for UEAndORT BackEnd. Set BackEnd to ENeuralBackEnd::Auto or ENeuralBackEnd::UEOnly for this platform."));
+#endif //WITH_FULL_NNI_SUPPORT
 	}
-
-	// Create session from model saved in ModelReadFromDiskInBytes (if not empty)
-	if (ModelReadFromDiskInBytes.Num() > 0)
+	// UEOnly
+	else if (BackEnd == ENeuralBackEnd::UEOnly)
 	{
-		// Read model from ModelReadFromDiskInBytesVector
-		ImplBackEndUEAndORT->Session = MakeUnique<Ort::Session>(*ImplBackEndUEAndORT->Environment, ModelReadFromDiskInBytes.GetData(), ModelReadFromDiskInBytes.Num(), *ImplBackEndUEAndORT->SessionOptions);
+		UE_LOG(LogNeuralNetworkInference, Warning, TEXT("UNeuralNetwork::Load(): Not implemented for BackEnd = %d yet."), (int32)BackEnd);
 	}
-	// Else
+	// Unknown
 	else
 	{
-		UE_LOG(LogNeuralNetworkInference, Warning, TEXT("ModelReadFromDiskInBytes was empty."));
-		return false;
+		UE_LOG(LogNeuralNetworkInference, Warning, TEXT("UNeuralNetwork::Load(): Unknown BackEnd = %d."), (int32)BackEnd);
 	}
 
-	ImplBackEndUEAndORT->ConfigureTensors(InputTensors, &AreInputTensorSizesVariable);
-	ImplBackEndUEAndORT->ConfigureTensors(OutputTensors);
-
-	bIsLoaded = true;
 	return bIsLoaded;
-
-#else
-	UE_LOG(LogNeuralNetworkInference, Warning, TEXT("UNeuralNetwork::Load(): Platform or Operating System not suported yet."));
-	return false;
-#endif //WITH_FULL_NNI_SUPPORT
 }
 
 bool UNeuralNetwork::IsLoaded() const
@@ -199,7 +133,10 @@ void UNeuralNetwork::SetDeviceType(const ENeuralDeviceType InDeviceType)
 	if (DeviceType != InDeviceType)
 	{
 		DeviceType = InDeviceType;
-		Load();
+		if (BackEnd == ENeuralBackEnd::UEAndORT)
+		{
+			Load();
+		}
 	}
 }
 
@@ -284,40 +221,34 @@ const FNeuralTensors& UNeuralNetwork::GetOutputTensors() const
 
 void UNeuralNetwork::Run()
 {
-#ifdef WITH_FULL_NNI_SUPPORT
 	// Sanity checks
 	if (!bIsLoaded)
 	{
 		UE_LOG(LogNeuralNetworkInference, Warning, TEXT("UNeuralNetwork::Run(): No architecture has been loaded yet. Run() will not work until IsLoaded() returns true."));
 		return;
 	}
-	// @todo: Temporarily disabled until we connect GPU input/output between UE and ORT
-	if (InputDeviceType == ENeuralDeviceType::GPU || OutputDeviceType == ENeuralDeviceType::GPU)
-	{
-		UE_LOG(LogNeuralNetworkInference, Warning, TEXT("UNeuralNetwork::Run(): InputDeviceType & OutputDeviceType must be set to CPU for now."));
-		return;
-	}
 
-	// Run UNeuralNetwork
-	if (SynchronousMode == ENeuralNetworkSynchronousMode::Synchronous)
+	// UEAndORT
+	if (BackEnd == ENeuralBackEnd::UEAndORT)
 	{
-		const FRedirectCoutAndCerrToUeLog RedirectCoutAndCerrToUeLog;
-		ImplBackEndUEAndORT->Session->Run(Ort::RunOptions{ nullptr },
-			InputTensors.GetTensorNames(), InputTensors.GetONNXRuntimeTensors(), InputTensors.GetNumberTensors(),
-			OutputTensors.GetTensorNames(), OutputTensors.GetONNXRuntimeTensors(), OutputTensors.GetNumberTensors());
+#ifdef WITH_FULL_NNI_SUPPORT
+		ImplBackEndUEAndORT->Run(OutputTensors, InputTensors, SynchronousMode, InputDeviceType, OutputDeviceType);
+#else
+		UE_LOG(LogNeuralNetworkInference, Warning, TEXT("UNeuralNetwork::Run(): Platform or Operating System not suported yet for UEAndORT BackEnd. Set BackEnd to ENeuralBackEnd::Auto or ENeuralBackEnd::UEOnly for this platform."));
+#endif //WITH_FULL_NNI_SUPPORT
 	}
-	else if (SynchronousMode == ENeuralNetworkSynchronousMode::Asynchronous)
+
+	// UEOnly
+	else if (BackEnd == ENeuralBackEnd::UEOnly)
 	{
-		UE_LOG(LogNeuralNetworkInference, Warning, TEXT("UNeuralNetwork::Run(): SynchronousMode = %d not implemented yet. Use SynchronousMode = Synchronous."), (int32)SynchronousMode);
+		UE_LOG(LogNeuralNetworkInference, Warning, TEXT("UNeuralNetwork::Run(): Not implemented for BackEnd = %d yet."), (int32)BackEnd);
 	}
+
+	// Unknown
 	else
 	{
-		UE_LOG(LogNeuralNetworkInference, Warning, TEXT("UNeuralNetwork::Run(): Unknown SynchronousMode = %d."), (int32)SynchronousMode);
+		UE_LOG(LogNeuralNetworkInference, Warning, TEXT("UNeuralNetwork::Run(): Unknown BackEnd = %d."), (int32)BackEnd);
 	}
-
-#else
-	UE_LOG(LogNeuralNetworkInference, Warning, TEXT("UNeuralNetwork::Run(): Platform or Operating System not suported yet."));
-#endif //WITH_FULL_NNI_SUPPORT
 }
 
 
