NeuralNetworkInference (NNI) plugin:
- NNILegacy merged into NNI step 4: *.imp file renamed as *.h and *.cpp.
- Added BackEndForCurrentPLatform.
- BackEnd set to Auto by default for new instances of UNN.

#rnx

diff --git a/Engine/Plugins/Experimental/NeuralNetworkInference/Source/NeuralNetworkInference/Private/NeuralNetwork.cpp b/Engine/Plugins/Experimental/NeuralNetworkInference/Source/NeuralNetworkInference/Private/NeuralNetwork.cpp
--- a/Engine/Plugins/Experimental/NeuralNetworkInference/Source/NeuralNetworkInference/Private/NeuralNetwork.cpp
+++ b/Engine/Plugins/Experimental/NeuralNetworkInference/Source/NeuralNetworkInference/Private/NeuralNetwork.cpp
@@ -1,17 +1,44 @@
 // Copyright Epic Games, Inc. All Rights Reserved.
 
 #include "NeuralNetwork.h"
+#include "NeuralNetworkImplBackEndUEAndORT.h"
 #include "NeuralNetworkInferenceUtils.h"
 #include "EditorFramework/AssetImportData.h"
 #include "HAL/FileManager.h"
+#include "Misc/FileHelper.h"
 #include "Misc/Paths.h"
 
 // Other files with UNeuralNetwork implementation (to be included only once and after the includes above)
-#include "NeuralNetworkImplBackEndUEAndORT.imp"
 #include "NeuralNetworkImplBackEndUEOnly.imp"
 
 
 
+/* FPrivateNeuralNetwork functions
+ *****************************************************************************/
+
+struct FPrivateNeuralNetwork
+{
+public:
+	static ENeuralBackEnd SetBackEndForCurrentPlatform(const ENeuralBackEnd InBackEnd);
+};
+
+ENeuralBackEnd FPrivateNeuralNetwork::SetBackEndForCurrentPlatform(const ENeuralBackEnd InBackEnd)
+{
+	// Auto
+	if (InBackEnd == ENeuralBackEnd::Auto)
+	{
+#ifdef WITH_FULL_NNI_SUPPORT
+		return ENeuralBackEnd::UEAndORT;
+#else //WITH_FULL_NNI_SUPPORT
+		return ENeuralBackEnd::UEOnly;
+#endif //WITH_FULL_NNI_SUPPORT
+	}
+	// Otherwise
+	return InBackEnd;
+}
+
+
+
 /* UNeuralNetwork structors
  *****************************************************************************/
 
@@ -20,12 +47,9 @@ UNeuralNetwork::UNeuralNetwork()
 	, InputDeviceType(ENeuralDeviceType::CPU)
 	, OutputDeviceType(ENeuralDeviceType::CPU)
 	, SynchronousMode(ENeuralNetworkSynchronousMode::Synchronous)
-#ifdef WITH_FULL_NNI_SUPPORT
-	, BackEnd(ENeuralBackEnd::UEAndORT)
-#else //WITH_FULL_NNI_SUPPORT
-	, BackEnd(ENeuralBackEnd::UEOnly)
-#endif //WITH_FULL_NNI_SUPPORT
+	, BackEnd(ENeuralBackEnd::Auto)
 	, bIsLoaded(false)
+	, BackEndForCurrentPlatform(FPrivateNeuralNetwork::SetBackEndForCurrentPlatform(BackEnd))
 {
 }
 
@@ -70,19 +94,19 @@ bool UNeuralNetwork::Load(const FString& InModelFilePath)
 	}
 
 	// UEAndORT
-	if (BackEnd == ENeuralBackEnd::UEAndORT)
+	if (BackEndForCurrentPlatform == ENeuralBackEnd::UEAndORT)
 	{
 		return Load();
 	}
 	// UEOnly
-	else if (BackEnd == ENeuralBackEnd::UEOnly)
+	else if (BackEndForCurrentPlatform == ENeuralBackEnd::UEOnly)
 	{
-		UE_LOG(LogNeuralNetworkInference, Warning, TEXT("UNeuralNetwork::Load(): Not implemented for BackEnd = %d yet."), (int32)BackEnd);
+		UE_LOG(LogNeuralNetworkInference, Warning, TEXT("UNeuralNetwork::Load(): Platform or Operating System not suported yet for [BackEnd,BackEndForCurrentPlatform] = [%d,%d]."), (int32)BackEnd, (int32)BackEndForCurrentPlatform);
 	}
 	// Unknown
 	else
 	{
-		UE_LOG(LogNeuralNetworkInference, Warning, TEXT("UNeuralNetwork::Load(): Unknown BackEnd = %d."), (int32)BackEnd);
+		UE_LOG(LogNeuralNetworkInference, Warning, TEXT("UNeuralNetwork::Load(): Unknown [BackEnd,BackEndForCurrentPlatform] = [%d,%d]."), (int32)BackEnd, (int32)BackEndForCurrentPlatform);
 	}
 
 	return bIsLoaded;
@@ -96,23 +120,19 @@ bool UNeuralNetwork::Load()
 	bIsLoaded = false;
 
 	// UEAndORT
-	if (BackEnd == ENeuralBackEnd::UEAndORT)
+	if (BackEndForCurrentPlatform == ENeuralBackEnd::UEAndORT)
 	{
-#ifdef WITH_FULL_NNI_SUPPORT
 		bIsLoaded = UNeuralNetwork::FImplBackEndUEAndORT::Load(ImplBackEndUEAndORT, InputTensors, OutputTensors, AreInputTensorSizesVariable, ModelReadFromDiskInBytes, ModelFullFilePath, GetDeviceType());
-#else
-		UE_LOG(LogNeuralNetworkInference, Warning, TEXT("UNeuralNetwork::Load(): Platform or Operating System not suported yet for UEAndORT BackEnd. Set BackEnd to ENeuralBackEnd::Auto or ENeuralBackEnd::UEOnly for this platform."));
-#endif //WITH_FULL_NNI_SUPPORT
 	}
 	// UEOnly
-	else if (BackEnd == ENeuralBackEnd::UEOnly)
+	else if (BackEndForCurrentPlatform == ENeuralBackEnd::UEOnly)
 	{
-		UE_LOG(LogNeuralNetworkInference, Warning, TEXT("UNeuralNetwork::Load(): Not implemented for BackEnd = %d yet."), (int32)BackEnd);
+		UE_LOG(LogNeuralNetworkInference, Warning, TEXT("UNeuralNetwork::Load(): Platform or Operating System not suported yet for [BackEnd,BackEndForCurrentPlatform] = [%d,%d]."), (int32)BackEnd, (int32)BackEndForCurrentPlatform);
 	}
 	// Unknown
 	else
 	{
-		UE_LOG(LogNeuralNetworkInference, Warning, TEXT("UNeuralNetwork::Load(): Unknown BackEnd = %d."), (int32)BackEnd);
+		UE_LOG(LogNeuralNetworkInference, Warning, TEXT("UNeuralNetwork::Load(): Unknown [BackEnd,BackEndForCurrentPlatform] = [%d,%d]."), (int32)BackEnd, (int32)BackEndForCurrentPlatform);
 	}
 
 	return bIsLoaded;
@@ -133,7 +153,7 @@ void UNeuralNetwork::SetDeviceType(const ENeuralDeviceType InDeviceType)
 	if (DeviceType != InDeviceType)
 	{
 		DeviceType = InDeviceType;
-		if (BackEnd == ENeuralBackEnd::UEAndORT)
+		if (BackEndForCurrentPlatform == ENeuralBackEnd::UEAndORT)
 		{
 			Load();
 		}
@@ -180,11 +200,20 @@ ENeuralBackEnd UNeuralNetwork::GetBackEnd() const
 	return BackEnd;
 }
 
+ENeuralBackEnd UNeuralNetwork::GetBackEndForCurrentPlatform() const
+{
+	return BackEndForCurrentPlatform;
+}
+
 void UNeuralNetwork::SetBackEnd(const ENeuralBackEnd InBackEnd)
 {
-	if (BackEnd != InBackEnd)
+	BackEnd = InBackEnd;
+	const ENeuralBackEnd NewBackEndForCurrentPlatform = FPrivateNeuralNetwork::SetBackEndForCurrentPlatform(BackEnd);
+	// Reload only required if BackEndForCurrentPlatform changes (regardless of whether BackEnd changed).
+	// BackEndForCurrentPlatform does not necesarily change if BackEnd changes. E.g., changing from UEAndORT into Auto in Windows will result in BackEndForCurrentPlatform = UEAndORT in both cases.
+	if (BackEndForCurrentPlatform != NewBackEndForCurrentPlatform)
 	{
-		BackEnd = InBackEnd;
+		BackEndForCurrentPlatform = NewBackEndForCurrentPlatform;
 		Load();
 	}
 }
@@ -229,25 +258,21 @@ void UNeuralNetwork::Run()
 	}
 
 	// UEAndORT
-	if (BackEnd == ENeuralBackEnd::UEAndORT)
+	if (BackEndForCurrentPlatform == ENeuralBackEnd::UEAndORT)
 	{
-#ifdef WITH_FULL_NNI_SUPPORT
 		ImplBackEndUEAndORT->Run(OutputTensors, InputTensors, SynchronousMode, InputDeviceType, OutputDeviceType);
-#else
-		UE_LOG(LogNeuralNetworkInference, Warning, TEXT("UNeuralNetwork::Run(): Platform or Operating System not suported yet for UEAndORT BackEnd. Set BackEnd to ENeuralBackEnd::Auto or ENeuralBackEnd::UEOnly for this platform."));
-#endif //WITH_FULL_NNI_SUPPORT
 	}
 
 	// UEOnly
-	else if (BackEnd == ENeuralBackEnd::UEOnly)
+	else if (BackEndForCurrentPlatform == ENeuralBackEnd::UEOnly)
 	{
-		UE_LOG(LogNeuralNetworkInference, Warning, TEXT("UNeuralNetwork::Run(): Not implemented for BackEnd = %d yet."), (int32)BackEnd);
+		UE_LOG(LogNeuralNetworkInference, Warning, TEXT("UNeuralNetwork::Run(): Platform or Operating System not suported yet for [BackEnd,BackEndForCurrentPlatform] = [%d,%d]."), (int32)BackEnd, (int32)BackEndForCurrentPlatform);
 	}
 
 	// Unknown
 	else
 	{
-		UE_LOG(LogNeuralNetworkInference, Warning, TEXT("UNeuralNetwork::Run(): Unknown BackEnd = %d."), (int32)BackEnd);
+		UE_LOG(LogNeuralNetworkInference, Warning, TEXT("UNeuralNetwork::Run(): Unknown [BackEnd,BackEndForCurrentPlatform] = [%d,%d]."), (int32)BackEnd, (int32)BackEndForCurrentPlatform);
 	}
 }
 
diff --git a/Engine/Plugins/Experimental/NeuralNetworkInference/Source/NeuralNetworkInference/Private/NeuralNetworkImplBackEndUEAndORT.cpp b/Engine/Plugins/Experimental/NeuralNetworkInference/Source/NeuralNetworkInference/Private/NeuralNetworkImplBackEndUEAndORT.cpp
new file mode 100644
--- /dev/null
+++ b/Engine/Plugins/Experimental/NeuralNetworkInference/Source/NeuralNetworkInference/Private/NeuralNetworkImplBackEndUEAndORT.cpp
@@ -0,0 +1,328 @@
+// Copyright Epic Games, Inc. All Rights Reserved.
+
+#include "NeuralNetworkImplBackEndUEAndORT.h"
+#include "NeuralNetworkInferenceUtils.h"
+
+
+
+/* UNeuralNetwork public functions
+ *****************************************************************************/
+
+//#if WITH_EDITOR
+//bool UNeuralNetwork::FImplBackEndUEAndORT::LoadFile(TSharedPtr<FImplBackEndUEAndORT>& InOutImplBackEndUEAndORT, FNeuralTensors& OutInputTensors, FNeuralTensors& OutOutputTensors,
+//	TArray<bool>& OutAreInputTensorSizesVariable, const TArray<uint8>& InModelReadFromDiskInBytes, const FString& InModelFullFilePath, const ENeuralDeviceType InDeviceType)
+//{
+//	// Initialize and configure InOutImplBackEndUEAndORT
+//	const FRedirectCoutAndCerrToUeLog RedirectCoutAndCerrToUeLog;
+//	if (!InitializedAndConfigureMembers(InOutImplBackEndUEAndORT, InModelFullFilePath, InDeviceType))
+//	{
+//		UE_LOG(LogNeuralNetworkInference, Warning, TEXT("UNeuralNetwork::Load(): InitializedAndConfigureMembers failed."));
+//		return false;
+//	}
+//
+//	// Note: This code uses the changes in the ONNX Runtime API, which are not needed for desktop platforms (where ONNX/ProtoBuf is supported)
+//	// Fill InModelReadFromDiskInBytes from ONNX/ORT file
+//	const FString FileExtension = FPaths::GetExtension(InModelFullFilePath, /*bIncludeDot*/ false);
+//	const char* const FilePathCharPtr = TCHAR_TO_ANSI(*InModelFullFilePath);
+//	// If ONNX file, turn into ORT format first
+//	if (FileExtension.Equals(TEXT("onnx"), ESearchCase::IgnoreCase))
+//	{
+//		FString ONNXPathPart, ONNXFilenamePart, ONNXExtensionPart;
+//		FPaths::Split(InModelFullFilePath, ONNXPathPart, ONNXFilenamePart, ONNXExtensionPart);
+//		const FString OutputORTOptimizedModelPath = ONNXPathPart / ONNXFilenamePart + TEXT(".ort");
+//#ifdef _WIN32
+//		InOutImplBackEndUEAndORT->SessionOptions->SetOptimizedModelFilePath(*OutputORTOptimizedModelPath);
+//#else
+//		InOutImplBackEndUEAndORT->SessionOptions->SetOptimizedModelFilePath(TCHAR_TO_ANSI(*OutputORTOptimizedModelPath));
+//#endif //_WIN32
+//		// ONNX --> ORT conversion
+//		// This session is just temporarily opened so the ORT model file can be generated
+//		InOutImplBackEndUEAndORT->Session = MakeUnique<Ort::Session>(*InOutImplBackEndUEAndORT->Environment,
+//#ifdef _WIN32
+//			*InModelFullFilePath,
+//#else
+//			FilePathCharPtr,
+//#endif
+//			*InOutImplBackEndUEAndORT->SessionOptions);
+//	
+//		// Read model from OutputORTOptimizedModelPath
+//		return Load(OutputORTOptimizedModelPath);
+//	}
+//	// Create session (it should be ORT file at this point), and read InModelReadFromDiskInBytes if not empty
+//	else if (FileExtension.Equals(TEXT("ort"), ESearchCase::IgnoreCase))
+//	{
+//		// Read model from InModelFullFilePath
+//		std::vector<uint8_t> OutputModelReadFromDiskInBytesVector;
+//		InOutImplBackEndUEAndORT->Session = MakeUnique<Ort::Session>(*InOutImplBackEndUEAndORT->Environment,
+//#ifdef _WIN32
+//			*InModelFullFilePath,
+//#else
+//			FilePathCharPtr,
+//#endif
+//			*InOutImplBackEndUEAndORT->SessionOptions, &OutputModelReadFromDiskInBytesVector);
+//	
+//		// Fill InModelReadFromDiskInBytes
+//		const int32 ArraySize = OutputModelReadFromDiskInBytesVector.size();
+//		if (ArraySize > 0)
+//		{
+//			InModelReadFromDiskInBytes.SetNumUninitialized(ArraySize);
+//			FMemory::Memcpy(InModelReadFromDiskInBytes.GetData(), &OutputModelReadFromDiskInBytesVector[0], ArraySize);
+//		}
+//	
+//		return Load();
+//	}
+//	// Else
+//	UE_LOG(LogNeuralNetworkInference, Warning, TEXT("UNeuralNetwork::Load(): Unknown file format \"%s\" used."), *FileExtension);
+//	return false;
+//}
+//#endif //WITH_EDITOR
+
+bool UNeuralNetwork::FImplBackEndUEAndORT::Load(TSharedPtr<FImplBackEndUEAndORT>& InOutImplBackEndUEAndORT, FNeuralTensors& OutInputTensors, FNeuralTensors& OutOutputTensors,
+	TArray<bool>& OutAreInputTensorSizesVariable, const TArray<uint8>& InModelReadFromDiskInBytes, const FString& InModelFullFilePath, const ENeuralDeviceType InDeviceType)
+{
+#ifdef WITH_FULL_NNI_SUPPORT
+	const FRedirectCoutAndCerrToUeLog RedirectCoutAndCerrToUeLog;
+
+	// Initialize and configure InOutImplBackEndUEAndORT
+	if (!UNeuralNetwork::FImplBackEndUEAndORT::InitializedAndConfigureMembers(InOutImplBackEndUEAndORT, InModelFullFilePath, InDeviceType))
+	{
+		UE_LOG(LogNeuralNetworkInference, Warning, TEXT("UNeuralNetwork::Load(): InitializedAndConfigureMembers failed."));
+		return false;
+	}
+
+	// Create session from model saved in InModelReadFromDiskInBytes (if not empty)
+	if (InModelReadFromDiskInBytes.Num() > 0)
+	{
+		// Read model from ModelReadFromDiskInBytesVector
+		InOutImplBackEndUEAndORT->Session = MakeUnique<Ort::Session>(*InOutImplBackEndUEAndORT->Environment, InModelReadFromDiskInBytes.GetData(), InModelReadFromDiskInBytes.Num(), *InOutImplBackEndUEAndORT->SessionOptions);
+	}
+	// Else
+	else
+	{
+		UE_LOG(LogNeuralNetworkInference, Warning, TEXT("InModelReadFromDiskInBytes was empty."));
+		return false;
+	}
+
+	InOutImplBackEndUEAndORT->ConfigureTensors(OutInputTensors, &OutAreInputTensorSizesVariable);
+	InOutImplBackEndUEAndORT->ConfigureTensors(OutOutputTensors);
+
+	return true;
+
+#else //WITH_FULL_NNI_SUPPORT
+	UE_LOG(LogNeuralNetworkInference, Warning, TEXT("UNeuralNetwork::Load(): Platform or Operating System not suported yet for UEAndORT BackEnd. Set BackEnd to ENeuralBackEnd::Auto (recommended) or ENeuralBackEnd::UEOnly for this platform."));
+	return false;
+#endif //WITH_FULL_NNI_SUPPORT
+}
+
+void UNeuralNetwork::FImplBackEndUEAndORT::Run(FNeuralTensors& InOutOutputTensors, const FNeuralTensors& InInputTensors, const ENeuralNetworkSynchronousMode InSynchronousMode, const ENeuralDeviceType InInputDeviceType, const ENeuralDeviceType InOutputDeviceType)
+{
+#ifdef WITH_FULL_NNI_SUPPORT
+	// @todo: Temporarily disabled until we connect GPU input/output between UE and ORT
+	if (InInputDeviceType == ENeuralDeviceType::GPU || InOutputDeviceType == ENeuralDeviceType::GPU)
+	{
+		UE_LOG(LogNeuralNetworkInference, Warning, TEXT("UNeuralNetwork::Run(): InputDeviceType & OutputDeviceType must be set to CPU for now."));
+		return;
+	}
+
+	// Run UNeuralNetwork
+	if (InSynchronousMode == ENeuralNetworkSynchronousMode::Synchronous)
+	{
+		const FRedirectCoutAndCerrToUeLog RedirectCoutAndCerrToUeLog;
+		Session->Run(Ort::RunOptions{ nullptr },
+			InInputTensors.GetTensorNames(), InInputTensors.GetONNXRuntimeTensors(), InInputTensors.GetNumberTensors(),
+			InOutOutputTensors.GetTensorNames(), InOutOutputTensors.GetONNXRuntimeTensors(), InOutOutputTensors.GetNumberTensors());
+	}
+	else if (InSynchronousMode == ENeuralNetworkSynchronousMode::Asynchronous)
+	{
+		UE_LOG(LogNeuralNetworkInference, Warning, TEXT("UNeuralNetwork::Run(): SynchronousMode = %d not implemented yet. Use SynchronousMode = Synchronous."), (int32)InSynchronousMode);
+	}
+	else
+	{
+		UE_LOG(LogNeuralNetworkInference, Warning, TEXT("UNeuralNetwork::Run(): Unknown SynchronousMode = %d."), (int32)InSynchronousMode);
+	}
+
+#else //WITH_FULL_NNI_SUPPORT
+	UE_LOG(LogNeuralNetworkInference, Warning, TEXT("UNeuralNetwork::Run(): Platform or Operating System not suported yet for UEAndORT BackEnd. Set BackEnd to ENeuralBackEnd::Auto or ENeuralBackEnd::UEOnly for this platform."));
+#endif //WITH_FULL_NNI_SUPPORT
+}
+
+
+
+/* UNeuralNetwork private functions
+ *****************************************************************************/
+
+#ifdef WITH_FULL_NNI_SUPPORT
+
+bool UNeuralNetwork::FImplBackEndUEAndORT::InitializedAndConfigureMembers(TSharedPtr<FImplBackEndUEAndORT>& InOutImplBackEndUEAndORT, const FString& InModelFullFilePath, const ENeuralDeviceType InDeviceType)
+{
+	// Initialize InOutImplBackEndUEAndORT
+	if (!InOutImplBackEndUEAndORT.IsValid())
+	{
+		InOutImplBackEndUEAndORT = MakeShared<FImplBackEndUEAndORT>();
+
+		// Set up ORT and create an environment
+		Ort::InitApi();
+		const char* const ModelFullFilePathCharPtr = TCHAR_TO_ANSI(*InModelFullFilePath);
+		InOutImplBackEndUEAndORT->Environment = MakeUnique<Ort::Env>(ORT_LOGGING_LEVEL_WARNING, ModelFullFilePathCharPtr); // Any unique string would work, it does not need to be the file path
+
+		InOutImplBackEndUEAndORT->Allocator = MakeUnique<Ort::AllocatorWithDefaultOptions>();
+	}
+
+	// Configure InOutImplBackEndUEAndORT
+	if (!InOutImplBackEndUEAndORT->ConfigureMembers(InDeviceType))
+	{
+		UE_LOG(LogNeuralNetworkInference, Warning, TEXT("UNeuralNetwork::InitializedAndConfigureMembers(): ConfigureMembers failed."));
+		return false;
+	}
+
+	return true;
+}
+
+bool UNeuralNetwork::FImplBackEndUEAndORT::ConfigureMembers(const ENeuralDeviceType InDeviceType)
+{
+	// Configure Session
+	SessionOptions = MakeUnique<Ort::SessionOptions>();
+
+	// Configure number threads
+	SessionOptions->SetIntraOpNumThreads(2);
+	// Uncomment if you want to change the priority of the threads, by default is TPri_Normal
+	//SessionOptions->SetPriorityOpThreads(EThreadPriority::TPri_Normal);
+
+	// Configure Provider
+	// GPU
+	if (InDeviceType == ENeuralDeviceType::GPU)
+	{
+#ifdef PLATFORM_WIN64
+		// ORT GPU (Direct ML)
+		SessionOptions->SetGraphOptimizationLevel(GraphOptimizationLevel::ORT_ENABLE_ALL); // ORT_ENABLE_ALL, ORT_ENABLE_EXTENDED, ORT_ENABLE_BASIC, ORT_DISABLE_ALL
+		if (OrtSessionOptionsAppendExecutionProvider_DML(*SessionOptions, 0))
+		{
+			UE_LOG(LogNeuralNetworkInference, Warning, TEXT("Some error occurred."));
+			return false;
+		}
+		return true; // @todo: Remove this line when NNI_HLSL is working
+#else
+		UE_LOG(LogNeuralNetworkInference, Warning, TEXT("UNeuralNetwork: GPU mode only supported in Windows for now. Please, switch to CPU or to Windows."));
+
+		//SessionOptions->SetGraphOptimizationLevel(GraphOptimizationLevel::ORT_DISABLE_ALL); // ORT_ENABLE_ALL, ORT_ENABLE_EXTENDED, ORT_ENABLE_BASIC, ORT_DISABLE_ALL
+		//if (OrtSessionOptionsAppendExecutionProvider_NNI_HLSL(*SessionOptions, 0))
+		//{
+		//	UE_LOG(LogNeuralNetworkInference, Warning, TEXT("Some error occurred."));
+		//	return false;
+		//}
+#endif //PLATFORM_WIN64
+	}
+	// CPU
+	//else // @todo: Uncomment this line when NNI_HLSL is working
+	{
+#ifdef WITH_NNI_CPU_NOT_RECOMMENDED
+		// NNI CPU (Deprecated)
+		SessionOptions->SetGraphOptimizationLevel(GraphOptimizationLevel::ORT_DISABLE_ALL); // ORT_ENABLE_ALL, ORT_ENABLE_EXTENDED, ORT_ENABLE_BASIC, ORT_DISABLE_ALL
+		if (OrtSessionOptionsAppendExecutionProvider_NNI_CPU(*SessionOptions))
+		{
+			UE_LOG(LogNeuralNetworkInference, Warning, TEXT("Some error occurred."));
+			return false;
+		}
+#else
+		// ORT CPU
+		SessionOptions->SetGraphOptimizationLevel(GraphOptimizationLevel::ORT_ENABLE_ALL); // ORT_ENABLE_ALL, ORT_ENABLE_EXTENDED, ORT_ENABLE_BASIC, ORT_DISABLE_ALL
+#endif //ORT_CPU
+	}
+
+	return true;
+}
+
+void UNeuralNetwork::FImplBackEndUEAndORT::ConfigureTensors(FNeuralTensors& OutTensors, TArray<bool>* OutAreInputTensorSizesVariable)
+{
+	const bool bIsInput = (OutAreInputTensorSizesVariable != nullptr);
+	TArray<const char*> TensorNames;
+	TArray<ENeuralDataType> TensorDataTypes;
+	TArray<TArray<int64>> TensorSizes;
+
+	const uint32 NumberTensors = bIsInput ? Session->GetInputCount() : Session->GetOutputCount();
+	if (OutAreInputTensorSizesVariable)
+	{
+		OutAreInputTensorSizesVariable->SetNum(NumberTensors);
+	}
+	for (uint32 TensorIndex = 0; TensorIndex < NumberTensors; ++TensorIndex)
+	{
+		// Get node name
+		{
+			const char* TensorName = bIsInput ? Session->GetInputName(TensorIndex, *Allocator) : Session->GetOutputName(TensorIndex, *Allocator);
+			TensorNames.Emplace(TensorName);
+		}
+
+		// Get node type
+		Ort::TypeInfo CurrentTypeInfo = bIsInput ? Session->GetInputTypeInfo(TensorIndex) : Session->GetOutputTypeInfo(TensorIndex);
+
+		Ort::TensorTypeAndShapeInfo CurrentTensorInfo = CurrentTypeInfo.GetTensorTypeAndShapeInfo();
+
+		{
+			ENeuralDataType TensorDataType;
+			{
+				const ONNXTensorElementDataType ONNXTensorElementDataTypeEnum = CurrentTensorInfo.GetElementType();
+				if (ONNXTensorElementDataTypeEnum == ONNX_TENSOR_ELEMENT_DATA_TYPE_FLOAT)
+				{
+					TensorDataType = ENeuralDataType::Float;
+				}
+				//else if (ONNXTensorElementDataTypeEnum == ONNX_TENSOR_ELEMENT_DATA_TYPE_INT32)
+				//{
+				//	TensorDataType = ENeuralDataType::Int32;
+				//}
+				//else if (ONNXTensorElementDataTypeEnum == ONNX_TENSOR_ELEMENT_DATA_TYPE_INT64)
+				//{
+				//	TensorDataType = ENeuralDataType::Int64;
+				//}
+				//else if (ONNXTensorElementDataTypeEnum == ONNX_TENSOR_ELEMENT_DATA_TYPE_UINT32)
+				//{
+				//	TensorDataType = ENeuralDataType::UInt32;
+				//}
+				//else if (ONNXTensorElementDataTypeEnum == ONNX_TENSOR_ELEMENT_DATA_TYPE_UINT64)
+				//{
+				//	TensorDataType = ENeuralDataType::UInt64;
+				//}
+				else
+				{
+					TensorDataType = ENeuralDataType::None;
+					UE_LOG(LogNeuralNetworkInference, Warning, TEXT("ONNXTensorElementDataTypeEnum = %d not implemented yet."), (int32)ONNXTensorElementDataTypeEnum);
+					return;
+				}
+			}
+			TensorDataTypes.Push(TensorDataType);
+		}
+
+		// Get input shapes/dims
+		{
+			TArray<int64> CurrentTensorSizes;
+			{
+				for (const int64_t CurrentTensorSize : CurrentTensorInfo.GetShape())
+				{
+					if (OutAreInputTensorSizesVariable)
+					{
+						(*OutAreInputTensorSizesVariable)[TensorIndex] |= (CurrentTensorSize < 0);
+					}
+					// Negative (variable) dimensions not implemented yet
+					if (CurrentTensorSize < 0)
+					{
+						CurrentTensorSizes.Push(1);
+						UE_LOG(LogNeuralNetworkInference, Display,
+							TEXT("Negative (i.e., variable) dimensions not allowed yet, hard-coded to 1. Let us know if you really need variable dimensions."
+								" Keep in mind that fixed sizes might allow additional optimizations and speedup of the network during Run()."));
+					}
+					else
+					{
+						CurrentTensorSizes.Push(CurrentTensorSize);
+					}
+				}
+			}
+			TensorSizes.Push(CurrentTensorSizes);
+		}
+
+		CurrentTypeInfo.release();
+
+	}
+
+	OutTensors.SetFromNetwork(TensorNames, TensorDataTypes, TensorSizes);
+}
+
+#endif //WITH_FULL_NNI_SUPPORT
diff --git a/Engine/Plugins/Experimental/NeuralNetworkInference/Source/NeuralNetworkInference/Private/NeuralNetworkImplBackEndUEAndORT.h b/Engine/Plugins/Experimental/NeuralNetworkInference/Source/NeuralNetworkInference/Private/NeuralNetworkImplBackEndUEAndORT.h
new file mode 100644
--- /dev/null
+++ b/Engine/Plugins/Experimental/NeuralNetworkInference/Source/NeuralNetworkInference/Private/NeuralNetworkImplBackEndUEAndORT.h
@@ -0,0 +1,56 @@
+// Copyright Epic Games, Inc. All Rights Reserved.
+
+#pragma once
+
+#include "NeuralNetwork.h"
+
+//#define WITH_NNI_CPU_NOT_RECOMMENDED // Only for debugging purposes
+
+#include "ThirdPartyWarningDisabler.h"
+NNI_THIRD_PARTY_INCLUDES_START
+#undef check
+#undef TEXT
+#ifdef WITH_FULL_NNI_SUPPORT
+	#include "RedirectCoutAndCerrToUeLog.h"
+
+	#include "onnxruntime/core/session/onnxruntime_cxx_api.h"
+	#ifdef PLATFORM_WIN64
+	#include "onnxruntime/core/providers/dml/dml_provider_factory.h"
+	#endif
+	#ifdef WITH_NNI_CPU_NOT_RECOMMENDED
+	#include "onnxruntime/core/providers/nni_cpu/nni_cpu_provider_factory.h"
+	#endif //WITH_NNI_CPU_NOT_RECOMMENDED
+#endif //WITH_FULL_NNI_SUPPORT
+NNI_THIRD_PARTY_INCLUDES_END
+
+
+
+struct UNeuralNetwork::FImplBackEndUEAndORT
+{
+public:
+#ifdef WITH_FULL_NNI_SUPPORT
+	TUniquePtr<Ort::Env> Environment;
+	TUniquePtr<Ort::Session> Session;
+	TUniquePtr<Ort::AllocatorWithDefaultOptions> Allocator;
+	TUniquePtr<Ort::SessionOptions> SessionOptions;
+#endif //WITH_FULL_NNI_SUPPORT
+
+//#if WITH_EDITOR
+//	static bool LoadFile(TSharedPtr<FImplBackEndUEAndORT>& InOutImplBackEndUEAndORT, FNeuralTensors& OutInputTensors, FNeuralTensors& OutOutputTensors,
+//		TArray<bool>& OutAreInputTensorSizesVariable, const TArray<uint8>& InModelReadFromDiskInBytes, const FString& InModelFullFilePath, const ENeuralDeviceType InDeviceType);
+//#endif //WITH_EDITOR
+
+	static bool Load(TSharedPtr<FImplBackEndUEAndORT>& InOutImplBackEndUEAndORT, FNeuralTensors& OutInputTensors, FNeuralTensors& OutOutputTensors,
+		TArray<bool>& OutAreInputTensorSizesVariable, const TArray<uint8>& InModelReadFromDiskInBytes, const FString& InModelFullFilePath, const ENeuralDeviceType InDeviceType);
+	
+	void Run(FNeuralTensors& InOutOutputTensors, const FNeuralTensors& InInputTensors, const ENeuralNetworkSynchronousMode InSynchronousMode, const ENeuralDeviceType InInputDeviceType, const ENeuralDeviceType InOutputDeviceType);
+
+#ifdef WITH_FULL_NNI_SUPPORT
+private:
+	static bool InitializedAndConfigureMembers(TSharedPtr<FImplBackEndUEAndORT>& InOutImplBackEndUEAndORT, const FString& InModelFullFilePath, const ENeuralDeviceType InDeviceType);
+
+	bool ConfigureMembers(const ENeuralDeviceType InDeviceType);
+
+	void ConfigureTensors(FNeuralTensors& OutTensors, TArray<bool>* OutAreInputTensorSizesVariable = nullptr);
+#endif //WITH_FULL_NNI_SUPPORT
+};
diff --git a/Engine/Plugins/Experimental/NeuralNetworkInference/Source/NeuralNetworkInference/Public/NeuralNetwork.h b/Engine/Plugins/Experimental/NeuralNetworkInference/Source/NeuralNetworkInference/Public/NeuralNetwork.h
--- a/Engine/Plugins/Experimental/NeuralNetworkInference/Source/NeuralNetworkInference/Public/NeuralNetwork.h
+++ b/Engine/Plugins/Experimental/NeuralNetworkInference/Source/NeuralNetworkInference/Public/NeuralNetwork.h
@@ -124,9 +124,14 @@ public:
 
 	/**
 	 * Getter and setter functions for BackEnd.
+	 * GetBackEnd()/GetBackEndForCurrentPlatform():
+	 * - If BackEnd == Auto, GetBackEnd() will return Auto and GetBackEndForCurrentPlatform() will return the actual BackEnd being used for the current platform (UEAndORT or UEOnly).
+	 * - If BackEnd != Auto, GetBackEnd() and GetBackEndForCurrentPlatform() will both return the same value (UEAndORT or UEOnly).
+	 * SetBackEnd() will modify both BackEnd and BackEndForCurrentPlatform.
 	 * @see ENeuralBackEnd for more details.
 	 */
 	ENeuralBackEnd GetBackEnd() const;
+	ENeuralBackEnd GetBackEndForCurrentPlatform() const;
 	void SetBackEnd(const ENeuralBackEnd InBackEnd);
 
 	/**
@@ -210,6 +215,13 @@ private:
 	UPROPERTY()
 	TArray<uint8> ModelReadFromDiskInBytes;
 
+	/**
+	 * If BackEnd != Auto, BackEndForCurrentPlatform will be equal to BackEnd.
+	 * Otherwise, BackEndForCurrentPlatform will be set to the optimal BackEnd given the current platform.
+	 * @see ENeuralBackEnd for more details.
+	 */
+	ENeuralBackEnd BackEndForCurrentPlatform;
+
 	/**
 	 * @see FOnAsyncRunCompleted and GetOnAsyncRunCompletedDelegate to understand OnAsyncRunCompletedDelegate.
 	 */
