NeuralNetworkInference (NNI) plugin:
- Quick rename from *.imp to *.h.

#rnx

diff --git a/Engine/Plugins/Experimental/NeuralNetworkInference/Source/NeuralNetworkInference/Private/NeuralNetwork.cpp b/Engine/Plugins/Experimental/NeuralNetworkInference/Source/NeuralNetworkInference/Private/NeuralNetwork.cpp
--- a/Engine/Plugins/Experimental/NeuralNetworkInference/Source/NeuralNetworkInference/Private/NeuralNetwork.cpp
+++ b/Engine/Plugins/Experimental/NeuralNetworkInference/Source/NeuralNetworkInference/Private/NeuralNetwork.cpp
@@ -2,15 +2,13 @@
 
 #include "NeuralNetwork.h"
 #include "NeuralNetworkImplBackEndUEAndORT.h"
+#include "NeuralNetworkImplBackEndUEOnly.h"
 #include "NeuralNetworkInferenceUtils.h"
 #include "EditorFramework/AssetImportData.h"
 #include "HAL/FileManager.h"
 #include "Misc/FileHelper.h"
 #include "Misc/Paths.h"
 
-// Other files with UNeuralNetwork implementation (to be included only once and after the includes above)
-#include "NeuralNetworkImplBackEndUEOnly.imp"
-
 
 
 /* FPrivateNeuralNetwork functions
diff --git a/Engine/Plugins/Experimental/NeuralNetworkInference/Source/NeuralNetworkInference/Private/NeuralNetworkImplBackEndUEOnly.h b/Engine/Plugins/Experimental/NeuralNetworkInference/Source/NeuralNetworkInference/Private/NeuralNetworkImplBackEndUEOnly.h
new file mode 100644
--- /dev/null
+++ b/Engine/Plugins/Experimental/NeuralNetworkInference/Source/NeuralNetworkInference/Private/NeuralNetworkImplBackEndUEOnly.h
@@ -0,0 +1,33 @@
+// Copyright Epic Games, Inc. All Rights Reserved.
+
+#pragma once
+
+#include "NeuralNetwork.h"
+
+#include "ModelProto.h"
+#include "NeuralOperator.h"
+#include "NeuralTensorManager.h"
+
+/* FImplBackEndUEOnly
+ *****************************************************************************/
+
+struct UNeuralNetwork::FImplBackEndUEOnly
+{
+	/**
+	 * It should always be false when loaded from uasset (FNeuralTensors are not auto-loaded to GPU).
+	 */
+	bool bAreTensorsInGpu;
+
+	FModelProto ModelProto;
+
+	/**
+	 * It contains a few TArray and TMaps for all FNeuralTensors (Input, Output, Intermediate(Not)Initialized, Weight).
+	 */
+	FNeuralTensorManager TensorManager;
+
+	/**
+	 * Only for the vanilla back end.
+	 * Set of operators that the network need to run on the Forward pass and that might need to run on the PostForward pass.
+	 */
+	TArray<TSharedPtr<FNeuralOperator>> Operators;
+};
