NeuralNetworkInference (NNI) plugin:
- UNN: GPU by default (rather than CPU).
- Updated UAssets.

#rnx

diff --git a/Engine/Plugins/Experimental/NeuralNetworkInference/Source/NeuralNetworkInference/Private/NeuralNetwork.cpp b/Engine/Plugins/Experimental/NeuralNetworkInference/Source/NeuralNetworkInference/Private/NeuralNetwork.cpp
--- a/Engine/Plugins/Experimental/NeuralNetworkInference/Source/NeuralNetworkInference/Private/NeuralNetwork.cpp
+++ b/Engine/Plugins/Experimental/NeuralNetworkInference/Source/NeuralNetworkInference/Private/NeuralNetwork.cpp
@@ -41,7 +41,7 @@ ENeuralBackEnd FPrivateNeuralNetwork::SetBackEndForCurrentPlatform(const ENeural
  *****************************************************************************/
 
 UNeuralNetwork::UNeuralNetwork()
-	: DeviceType(ENeuralDeviceType::CPU)
+	: DeviceType(ENeuralDeviceType::GPU)
 	, InputDeviceType(ENeuralDeviceType::CPU)
 	, OutputDeviceType(ENeuralDeviceType::CPU)
 	, SynchronousMode(ENeuralNetworkSynchronousMode::Synchronous)
