NeuralNetworkInference (NNI) plugin:
- ModelProtoReader: C++ WITH_WINDOWS macros disabled, so it can run in all platforms.

[REDACTED_HANDLE], this should remove all warnings from the QA test in Linux/Mac.

#rnx

diff --git a/Engine/Plugins/Experimental/NeuralNetworkInference/Source/NeuralNetworkInference/Private/GraphProtoToNeuralNetworkConverter.cpp b/Engine/Plugins/Experimental/NeuralNetworkInference/Source/NeuralNetworkInference/Private/GraphProtoToNeuralNetworkConverter.cpp
--- a/Engine/Plugins/Experimental/NeuralNetworkInference/Source/NeuralNetworkInference/Private/GraphProtoToNeuralNetworkConverter.cpp
+++ b/Engine/Plugins/Experimental/NeuralNetworkInference/Source/NeuralNetworkInference/Private/GraphProtoToNeuralNetworkConverter.cpp
@@ -8,9 +8,7 @@
 #include "NeuralTensorManager.h"
 
 #if WITH_EDITOR
-#if PLATFORM_WINDOWS
 #include "ModelProtoFileReader.h"
-#endif //PLATFORM_WINDOWS
 #endif //WITH_EDITOR
 
 
diff --git a/Engine/Plugins/Experimental/NeuralNetworkInference/Source/NeuralNetworkInference/Private/NeuralNetworkLegacy.cpp b/Engine/Plugins/Experimental/NeuralNetworkInference/Source/NeuralNetworkInference/Private/NeuralNetworkLegacy.cpp
--- a/Engine/Plugins/Experimental/NeuralNetworkInference/Source/NeuralNetworkInference/Private/NeuralNetworkLegacy.cpp
+++ b/Engine/Plugins/Experimental/NeuralNetworkInference/Source/NeuralNetworkInference/Private/NeuralNetworkLegacy.cpp
@@ -11,9 +11,7 @@
 #include "RHI.h"
 
 #if WITH_EDITOR
-#if PLATFORM_WINDOWS
 #include "ModelProtoFileReader.h"
-#endif //PLATFORM_WINDOWS
 #endif //WITH_EDITOR
 
 
@@ -116,10 +114,8 @@ UAssetImportData* UNeuralNetworkLegacy::GetAndMaybeCreateAssetImportData()
 #endif // WITH_EDITOR
 
 #if WITH_EDITOR
-bool UNeuralNetworkLegacy::Load(const FString& InFilePath)
+bool UNeuralNetworkLegacy::Load(const FString& InModelFilePath)
 {
-#if PLATFORM_WINDOWS
-
 	DECLARE_SCOPE_CYCLE_COUNTER(TEXT("UNeuralNetworkLegacy_Load_File"), STAT_UNeuralNetworkLegacy_Load_File, STATGROUP_MachineLearning);
 
 	// Clean previous networks
@@ -133,25 +129,19 @@ bool UNeuralNetworkLegacy::Load(const FString& InFilePath)
 	}
 
 	// Read ModelProto
-	if (!FModelProtoFileReader::ReadModelProtoFromFile(ModelProto, InFilePath) || !ModelProto.IsLoaded())
+	if (!FModelProtoFileReader::ReadModelProtoFromFile(ModelProto, InModelFilePath) || !ModelProto.IsLoaded())
 	{
-		UE_LOG(LogNeuralNetworkInference, Warning, TEXT("UNeuralNetworkLegacy::Load(): Model could not be loaded from %s."), *InFilePath);
+		UE_LOG(LogNeuralNetworkInference, Warning, TEXT("UNeuralNetworkLegacy::Load(): Model could not be loaded from %s."), *InModelFilePath);
 		return false;
 	}
 	// Turn ModelProto into Operators
-	bIsLoaded = FGraphProtoToNeuralNetworkConverter::Translate(Operators, TensorManager, ModelProto.GetGraph(), InFilePath);
+	bIsLoaded = FGraphProtoToNeuralNetworkConverter::Translate(Operators, TensorManager, ModelProto.GetGraph(), InModelFilePath);
 	if (!TensorManager.IsLoaded())
 	{
-		UE_LOG(LogNeuralNetworkInference, Warning, TEXT("UNeuralNetworkLegacy::Load(): TensorManager could not be loaded from %s."), *InFilePath);
+		UE_LOG(LogNeuralNetworkInference, Warning, TEXT("UNeuralNetworkLegacy::Load(): TensorManager could not be loaded from %s."), *InModelFilePath);
 		return false;
 	}
 	return bIsLoaded;
-
-#else //PLATFORM_WINDOWS
-	UE_LOG(LogNeuralNetworkInference, Warning, TEXT("UNeuralNetworkLegacy::Load(): Only implemented for Windows."),
-		*InFilePath, ModelProto.IsLoaded());
-	return false;
-#endif //PLATFORM_WINDOWS
 }
 #endif //WITH_EDITOR
 
diff --git a/Engine/Plugins/Experimental/NeuralNetworkInference/Source/NeuralNetworkInference/Private/NeuralTensor.cpp b/Engine/Plugins/Experimental/NeuralNetworkInference/Source/NeuralNetworkInference/Private/NeuralTensor.cpp
--- a/Engine/Plugins/Experimental/NeuralNetworkInference/Source/NeuralNetworkInference/Private/NeuralTensor.cpp
+++ b/Engine/Plugins/Experimental/NeuralNetworkInference/Source/NeuralNetworkInference/Private/NeuralTensor.cpp
@@ -453,7 +453,6 @@ bool FNeuralTensor::SetFromTensorProto(const FTensorProto* const InTensorProto,
 	else if (!InTensorProto->ExternalData.IsEmpty())
 	{
 #if WITH_EDITOR
-#if PLATFORM_WINDOWS
 		UE_LOG(LogNeuralNetworkInference, Warning, TEXT("FNeuralTensor::SetFromTensorProto(): DEPRECATED CODE, OTXT no longer functional."));
 		// // Sanity check
 		// if (InTensorProto->ExternalData[0].Key != TEXT("location") || InTensorProto->ExternalData[0].Value.Len() < 1)
@@ -469,10 +468,6 @@ bool FNeuralTensor::SetFromTensorProto(const FTensorProto* const InTensorProto,
 		// 	UE_LOG(LogNeuralNetworkInference, Warning, TEXT("FNeuralTensor::SetFromTensorProto(): Could not read binary file: %s."), *BinaryWeightFilePath);
 		// 	return false;
 		// }
-#else //PLATFORM_WINDOWS
-		UE_LOG(LogNeuralNetworkInference, Warning, TEXT("FNeuralTensor::SetFromTensorProto(): Only implemented in Windows."));
-		return false;
-#endif //PLATFORM_WINDOWS
 #else //WITH_EDITOR
 		UE_LOG(LogNeuralNetworkInference, Warning, TEXT("FNeuralTensor::SetFromTensorProto(): InTensorProto->ExternalData should never be used in non-Editor mode."));
 		return false;
