NeuralNetworkInference (NNI) plugin:
- UE-only back end partially implemented for UNN.
- Fixed Horde errors for Mac/Linux.

#rnx

diff --git a/Engine/Plugins/Experimental/NeuralNetworkInference/Source/NeuralNetworkInference/NeuralNetworkInference.Build.cs b/Engine/Plugins/Experimental/NeuralNetworkInference/Source/NeuralNetworkInference/NeuralNetworkInference.Build.cs
--- a/Engine/Plugins/Experimental/NeuralNetworkInference/Source/NeuralNetworkInference/NeuralNetworkInference.Build.cs
+++ b/Engine/Plugins/Experimental/NeuralNetworkInference/Source/NeuralNetworkInference/NeuralNetworkInference.Build.cs
@@ -50,16 +50,12 @@ public class NeuralNetworkInference : ModuleRules
 		// Editor-only
 		if (Target.bBuildEditor)
 		{
-			// Win64-only
-			if (Target.Platform == UnrealTargetPlatform.Win64)
-			{
-				PrivateDependencyModuleNames.AddRange
-					(
-					new string[] {
-						"ModelProtoFileReader"
-					}
-				);
-			}
+			PrivateDependencyModuleNames.AddRange
+				(
+				new string[] {
+					"ModelProtoFileReader"
+				}
+			);
 		}
 
 		if (bIsORTSupported)
diff --git a/Engine/Plugins/Experimental/NeuralNetworkInference/Source/NeuralNetworkInference/Private/NeuralNetwork.cpp b/Engine/Plugins/Experimental/NeuralNetworkInference/Source/NeuralNetworkInference/Private/NeuralNetwork.cpp
--- a/Engine/Plugins/Experimental/NeuralNetworkInference/Source/NeuralNetworkInference/Private/NeuralNetwork.cpp
+++ b/Engine/Plugins/Experimental/NeuralNetworkInference/Source/NeuralNetworkInference/Private/NeuralNetwork.cpp
@@ -99,7 +99,7 @@ bool UNeuralNetwork::Load(const FString& InModelFilePath)
 	// UEOnly
 	else if (BackEndForCurrentPlatform == ENeuralBackEnd::UEOnly)
 	{
-		UE_LOG(LogNeuralNetworkInference, Warning, TEXT("UNeuralNetwork::Load(): Platform or Operating System not suported yet for [BackEnd,BackEndForCurrentPlatform] = [%d,%d]."), (int32)BackEnd, (int32)BackEndForCurrentPlatform);
+		bIsLoaded = UNeuralNetwork::FImplBackEndUEOnly::Load(/*ModelReadFromDiskInBytes*/ ModelFullFilePath);
 	}
 	// Unknown
 	else
diff --git a/Engine/Plugins/Experimental/NeuralNetworkInference/Source/NeuralNetworkInference/Private/NeuralNetworkImplBackEndUEOnly.cpp b/Engine/Plugins/Experimental/NeuralNetworkInference/Source/NeuralNetworkInference/Private/NeuralNetworkImplBackEndUEOnly.cpp
new file mode 100644
--- /dev/null
+++ b/Engine/Plugins/Experimental/NeuralNetworkInference/Source/NeuralNetworkInference/Private/NeuralNetworkImplBackEndUEOnly.cpp
@@ -0,0 +1,342 @@
+// Copyright Epic Games, Inc. All Rights Reserved.
+
+#include "NeuralNetworkImplBackEndUEOnly.h"
+//#include "NeuralOperator.h"
+//#include "GraphProtoToNeuralNetworkConverter.h"
+//#include "NeuralNetworkInferenceUtils.h"
+//#include "EditorFramework/AssetImportData.h"
+//#include "HAL/FileManager.h"
+//#include "RenderGraphBuilder.h"
+//#include "RenderingThread.h"
+//#include "RHI.h"
+
+//#if WITH_EDITOR
+//#include "ModelProtoFileReader.h"
+//#endif //WITH_EDITOR
+
+
+
+/* FImplBackEndUEOnly public functions
+ *****************************************************************************/
+
+#if WITH_EDITOR
+
+bool UNeuralNetwork::FImplBackEndUEOnly::Load(const FString& InModelFilePath)
+{
+//	DECLARE_SCOPE_CYCLE_COUNTER(TEXT("UNeuralNetworkLegacy_Load_File"), STAT_UNeuralNetworkLegacy_Load_File, STATGROUP_MachineLearning);
+//
+//	// Clean previous networks
+//	if (bIsLoaded)
+//	{
+//		UE_LOG(LogNeuralNetworkInference, Warning, TEXT("UNeuralNetwork::FImplBackEndUEOnly::Load(): A model was previously loaded, removing it and reloading the new model."));
+//		Operators.Empty();
+//		ModelProto = FModelProto();
+//		bIsLoaded = false;
+//		bAreTensorsInGpu = false;
+//	}
+//
+//	// Read ModelProto
+//	if (!FModelProtoFileReader::ReadModelProtoFromFile(ModelProto, InModelFilePath) || !ModelProto.IsLoaded())
+//	{
+//		UE_LOG(LogNeuralNetworkInference, Warning, TEXT("UNeuralNetwork::FImplBackEndUEOnly::Load(): Model could not be loaded from %s."), *InModelFilePath);
+//		return false;
+//	}
+//	// Turn ModelProto into Operators
+//	bIsLoaded = FGraphProtoToNeuralNetworkConverter::Translate(Operators, TensorManager, ModelProto.GetGraph(), InModelFilePath);
+//	if (!TensorManager.IsLoaded())
+//	{
+//		UE_LOG(LogNeuralNetworkInference, Warning, TEXT("UNeuralNetwork::FImplBackEndUEOnly::Load(): TensorManager could not be loaded from %s."), *InModelFilePath);
+		return false;
+//	}
+	//return true;
+}
+#endif //WITH_EDITOR
+
+bool UNeuralNetwork::FImplBackEndUEOnly::Load(FNeuralTensorManager& InTensorManager, const TArray<TSharedPtr<FNeuralOperator>>& InOperators)
+{
+	//DECLARE_SCOPE_CYCLE_COUNTER(TEXT("UNeuralNetworkLegacy_Load"), STAT_UNeuralNetworkLegacy_Load, STATGROUP_MachineLearning);
+	//if (!InTensorManager.IsLoaded())
+	//{
+	//	UE_LOG(LogNeuralNetworkInference, Warning, TEXT("UNeuralNetwork::FImplBackEndUEOnly::Load(): TensorManager could not be loaded."));
+	//}
+	//Swap(TensorManager, InTensorManager);
+	//Operators = InOperators;
+	//return (Operators.Num() > 0 && TensorManager.IsLoaded());
+return false;
+}
+
+//const TArray<FNeuralTensor>& UNeuralNetwork::FImplBackEndUEOnly::GetTensors() const
+//{
+//	return TensorManager.GetTensors();
+//}
+//
+//void UNeuralNetwork::FImplBackEndUEOnly::SetInputFromTensorCopy(const FNeuralTensor& InTensorMap)
+//{
+//	return TensorManager.SetInputFromTensorCopy(InTensorMap);
+//}
+//
+//void UNeuralNetwork::FImplBackEndUEOnly::SetInputFromTensorMapCopy(const TMap<FString, FNeuralTensor>& InTensorMap)
+//{
+//	return TensorManager.SetInputFromTensorMapCopy(InTensorMap);
+//}
+//
+//TMap<FString, void*> UNeuralNetwork::FImplBackEndUEOnly::CreateInputDataPointersMutable()
+//{
+//	return TensorManager.CreateInputDataPointersMutable();
+//}
+//
+////FRDGBufferUAVRef UNeuralNetwork::FImplBackEndUEOnly::GetInputBufferUAVRef()
+////{
+////	return TensorManager.GetInputBufferUAVRef();
+////}
+////
+////TMap<FString, FRDGBufferUAVRef> UNeuralNetwork::FImplBackEndUEOnly::CreateInputBufferUAVRefs()
+////{
+////	return TensorManager.CreateInputBufferUAVRefs();
+////}
+//
+//const FNeuralTensor& UNeuralNetwork::FImplBackEndUEOnly::GetInputTensor() const
+//{
+//	return TensorManager.GetInputTensor();
+//}
+//
+//const TMap<FString, int32>& UNeuralNetwork::FImplBackEndUEOnly::GetInputNameIndexMap() const
+//{
+//	return TensorManager.GetInputNameIndexMap();
+//}
+//
+//const FNeuralTensor& UNeuralNetwork::FImplBackEndUEOnly::GetOutputTensor() const
+//{
+//	return TensorManager.GetOutputTensor();
+//}
+//
+//const TMap<FString, int32>& UNeuralNetwork::FImplBackEndUEOnly::GetOutputNameIndexMap() const
+//{
+//	return TensorManager.GetOutputNameIndexMap();
+//}
+//
+//const FRDGBufferSRVRef UNeuralNetwork::FImplBackEndUEOnly::GetOutputBufferSRVRef() const
+//{
+//	return TensorManager.GetOutputBufferSRVRef();
+//}
+//
+//TMap<FString, const FRDGBufferSRVRef> UNeuralNetwork::FImplBackEndUEOnly::CreateOutputBufferSRVRefs() const
+//{
+//	return TensorManager.CreateOutputBufferSRVRefs();
+//}
+//
+//TMap<FString, FNeuralTensor> UNeuralNetwork::FImplBackEndUEOnly::CreateInputTensorMap() const
+//{
+//	return TensorManager.CreateInputTensorMap();
+//}
+//
+//TMap<FString, FNeuralTensor> UNeuralNetwork::FImplBackEndUEOnly::CreateOutputTensorMap() const
+//{
+//	return TensorManager.CreateOutputTensorMap();
+//}
+//
+//void UNeuralNetwork::FImplBackEndUEOnly::Run(const ENeuralNetworkSynchronousMode InSynchronousMode, const ENeuralDeviceType InInputDeviceType, const ENeuralDeviceType InOutputDeviceType, const bool bRunGPUEmptyOnlyForProfiling)
+//{
+//	DECLARE_SCOPE_CYCLE_COUNTER(TEXT("UNeuralNetworkLegacy_Run"), STAT_UNeuralNetworkLegacy_Run, STATGROUP_MachineLearning);
+//	// Sanity checks
+//	if (!bIsLoaded)
+//	{
+//		UE_LOG(LogNeuralNetworkInference, Warning, TEXT("UNeuralNetwork::FImplBackEndUEOnly::Run(): No architecture has been loaded yet. Run() will not work until IsLoaded() returns true."));
+//		return;
+//	}
+//	// Run UNeuralNetworkLegacy
+//	if (Operators.Num() > 0)
+//	{
+//		// Run graph
+//		if (DeviceType == ENeuralDeviceType::CPU)
+//		{
+//			DECLARE_SCOPE_CYCLE_COUNTER(TEXT("UNeuralNetworkLegacy_Run::Forward_CPU"), STAT_UNeuralNetworkLegacy_Run_Forward_CPU, STATGROUP_MachineLearning);
+//			// Run each operator forward pass
+//			for (TSharedPtr<FNeuralOperator>& Operator : Operators)
+//			{
+//				Operator->ForwardCPU();
+//			}
+//			// Run each operator post forward pass
+//			for (TSharedPtr<FNeuralOperator>& Operator : Operators)
+//			{
+//				Operator->PostForwardCPU();
+//			}
+//		}
+//		else if (DeviceType == ENeuralDeviceType::GPU)
+//		{
+//			DECLARE_SCOPE_CYCLE_COUNTER(TEXT("UNeuralNetworkLegacy_Run::Forward_GPU"), STAT_UNeuralNetworkLegacy_Run_Forward_GPU, STATGROUP_MachineLearning);
+//			// Sanity check
+//			if (TensorManager.GetTensorsMutable().Num() < 1)
+//			{
+//				UE_LOG(LogNeuralNetworkInference, Warning, TEXT("UNeuralNetwork::FImplBackEndUEOnly::Run(): Tensors.Num() = %d (should be > 0)."), TensorManager.GetTensorsMutable().Num());
+//				return;
+//			}
+//
+//			// On RHI thread
+//			ENQUEUE_RENDER_COMMAND(UNeuralNetworkLegacy_Run_RenderThread)(
+//				[this, bRunGPUEmptyOnlyForProfiling, InSynchronousMode, InInputDeviceType, InOutputDeviceType](FRHICommandListImmediate& RHICmdList)
+//				{
+//					FMemMark Mark(FMemStack::Get());
+//					FRDGBuilder GraphBuilder(RHICmdList, RDG_EVENT_NAME("UNeuralNetwork::FImplBackEndUEOnly::Run()"));
+//
+//					// Move memory from CPU to GPU
+//					TArray<FNeuralTensor>& Tensors = TensorManager.GetTensorsMutable();
+//					const bool bIsInputInCPU = (InInputDeviceType == ENeuralDeviceType::CPU);
+//					// Move only input tensors to GPU (once per UNeuralNetwork::FImplBackEndUEOnly::Run())
+//					if (bAreTensorsInGpu)
+//					{
+//						// If input is on CPU, move inputs to GPU
+//						if (bIsInputInCPU)
+//						{
+//							for (const int32 InputIndex : TensorManager.GetInputIndexes())
+//							{
+//								Tensors[InputIndex].ToGPU_RenderThread(&GraphBuilder);
+//							}
+//							for (const int32 NonInputIndex : TensorManager.GetNonInputIndexes())
+//							{
+//								Tensors[NonInputIndex].UpdateSRVAndOrUAV_RenderThread(&GraphBuilder);
+//							}
+//						}
+//						// If input is already on GPU, just refresh all w.r.t. GraphBuilder
+//						else
+//						{
+//							for (FNeuralTensor& Tensor : Tensors)
+//							{
+//								Tensor.UpdateSRVAndOrUAV_RenderThread(&GraphBuilder);
+//							}
+//						}
+//					}
+//					// Move all (input, intermediate, weight, output) tensors to GPU and also call Operators.ToGPU() to move their auxiliary memory into GPU
+//					// Only once per UNeuralNetworkLegacy instance, or until Load() is called again
+//					else
+//					{
+//						// If input is on CPU, move inputs to GPU
+//						if (bIsInputInCPU)
+//						{
+//							// Move tensors to GPU
+//							for (FNeuralTensor& Tensor : Tensors)
+//							{
+//								Tensor.ToGPU_RenderThread(&GraphBuilder);
+//							}
+//						}
+//						// If input is already on GPU, just refresh inputs
+//						else
+//						{
+//							for (const int32 InputIndex : TensorManager.GetInputIndexes())
+//							{
+//								Tensors[InputIndex].UpdateSRVAndOrUAV_RenderThread(&GraphBuilder);
+//							}
+//							for (const int32 NonInputIndex : TensorManager.GetNonInputIndexes())
+//							{
+//								Tensors[NonInputIndex].ToGPU_RenderThread(&GraphBuilder);
+//							}
+//						}
+//						// Operators->ToGPU()
+//						for (TSharedPtr<FNeuralOperator>& Operator : Operators)
+//						{
+//							Operator->ToGPU_RenderThread();
+//						}
+//						// bAreTensorsInGpu is true
+//						bAreTensorsInGpu = true;
+//					}
+//
+//					if (!bRunGPUEmptyOnlyForProfiling)
+//					{
+//						// Run each operator forward pass
+//						for (TSharedPtr<FNeuralOperator>& Operator : this->Operators)
+//						{
+//							Operator->ForwardGPU_RenderThread(&GraphBuilder);
+//						}
+//
+//						// Run each operator post forward pass
+//						for (TSharedPtr<FNeuralOperator>& Operator : this->Operators)
+//						{
+//							Operator->PostForwardGPU_RenderThread(&GraphBuilder);
+//						}
+//					}
+//
+//					// Move memory from GPU to CPU
+//					const bool bIsOutputInCPU = (InOutputDeviceType == ENeuralDeviceType::CPU);
+//					if (bIsOutputInCPU)
+//					{
+//						for (const int32 OutputIndex : TensorManager.GetOutputIndexes())
+//						{
+//							Tensors[OutputIndex].ToCPU_RenderThread(&GraphBuilder);
+//						}
+//					}
+//
+//					// Broadcast delegates (from the render thread)
+//					if (InSynchronousMode == ENeuralNetworkSynchronousMode::Asynchronous)
+//					{
+//						GraphBuilder.AddPass(
+//							RDG_EVENT_NAME("Async delegate broadcast"),
+//							ERDGPassFlags::None,
+//							[this](FRHICommandListImmediate& RHICmdList)
+//						{
+//							OnAsyncRunCompletedDelegate.ExecuteIfBound();
+//						});
+//					}
+//
+//					// Execute Render Graph
+//					GraphBuilder.Execute();
+//				}
+//			);
+//
+//			// Block thread until GPU has finished
+//			if (InSynchronousMode == ENeuralNetworkSynchronousMode::Synchronous)
+//			{
+//				FNeuralNetworkInferenceUtils::WaitUntilRHIFinished();
+//			}
+//		}
+//		else
+//		{
+//			UE_LOG(LogNeuralNetworkInference, Warning, TEXT("UNeuralNetwork::FImplBackEndUEOnly::Run(): Unknown DeviceType = %d."), (int32)DeviceType);
+//		}
+//	}
+//	else
+//	{
+//		UE_LOG(LogNeuralNetworkInference, Warning, TEXT("UNeuralNetwork::FImplBackEndUEOnly::Run() called with an empty model."));
+//	}
+//}
+//
+//FString UNeuralNetwork::FImplBackEndUEOnly::ToString() const
+//{
+//	// Add GraphProto
+//	FString String = ModelProto.GetGraph().ToString();
+//	// Add FNeuralTensor(s)
+//	String += TEXT("TensorManager:\n");
+//	const TMap<FString, int32>& NameIndexMap = TensorManager.GetNameIndexMap();
+//	const TArray<FNeuralTensor>& Tensors = TensorManager.GetTensors();
+//	if (NameIndexMap.Num() > 0)
+//	{
+//		for (const auto& NameIndexPair : NameIndexMap)
+//		{
+//			String += FString::Format(TEXT(" -{0}: {1}\n"), { NameIndexPair.Key, Tensors[NameIndexPair.Value].ToString(20) });
+//		}
+//	}
+//	else
+//	{
+//		for (const FNeuralTensor& Tensor : Tensors)
+//		{
+//			String += FString::Format(TEXT(" -{0}\n"), { Tensor.ToString(20) });
+//		}
+//	}
+//	String += TEXT("InputTensorMap:\n");
+//	for (const auto& NameIndexPair : TensorManager.GetInputNameIndexMap())
+//	{
+//		String += FString::Format(TEXT(" -{0}: {1}\n"), { NameIndexPair.Key, Tensors[NameIndexPair.Value].ToString(20) });
+//	}
+//	String += TEXT("OutputTensorMap:\n");
+//	for (const auto& NameIndexPair : TensorManager.GetOutputNameIndexMap())
+//	{
+//		String += FString::Format(TEXT(" -{0}: {1}\n"), { NameIndexPair.Key, Tensors[NameIndexPair.Value].ToString(20) });
+//	}
+//	// Add FNeuralOperator(s)
+//	String += TEXT("Operators:\n");
+//	for (const TSharedPtr<FNeuralOperator>& Operator : Operators)
+//	{
+//		String += FString::Format(TEXT(" -{0}\n"), { Operator->ToString() });
+//	}
+//	// Return result
+//	return String;
+//}
diff --git a/Engine/Plugins/Experimental/NeuralNetworkInference/Source/NeuralNetworkInference/Private/NeuralNetworkImplBackEndUEOnly.h b/Engine/Plugins/Experimental/NeuralNetworkInference/Source/NeuralNetworkInference/Private/NeuralNetworkImplBackEndUEOnly.h
--- a/Engine/Plugins/Experimental/NeuralNetworkInference/Source/NeuralNetworkInference/Private/NeuralNetworkImplBackEndUEOnly.h
+++ b/Engine/Plugins/Experimental/NeuralNetworkInference/Source/NeuralNetworkInference/Private/NeuralNetworkImplBackEndUEOnly.h
@@ -30,4 +30,10 @@ struct UNeuralNetwork::FImplBackEndUEOnly
 	 * Set of operators that the network need to run on the Forward pass and that might need to run on the PostForward pass.
 	 */
 	TArray<TSharedPtr<FNeuralOperator>> Operators;
+
+ #if WITH_EDITOR
+ 	static bool Load(const FString& InModelFullFilePath);
+ #endif //WITH_EDITOR
+
+	static bool Load(FNeuralTensorManager& InTensorManager, const TArray<TSharedPtr<FNeuralOperator>>& InOperators);
 };
diff --git a/Engine/Plugins/Experimental/NeuralNetworkInference/Source/NeuralNetworkInference/Private/NeuralNetworkLegacy.cpp b/Engine/Plugins/Experimental/NeuralNetworkInference/Source/NeuralNetworkInference/Private/NeuralNetworkLegacy.cpp
--- a/Engine/Plugins/Experimental/NeuralNetworkInference/Source/NeuralNetworkInference/Private/NeuralNetworkLegacy.cpp
+++ b/Engine/Plugins/Experimental/NeuralNetworkInference/Source/NeuralNetworkInference/Private/NeuralNetworkLegacy.cpp
@@ -42,7 +42,7 @@ void UNeuralNetworkLegacy::PostInitProperties()
 	{
 		GetAndMaybeCreateAssetImportData();
 	}
-#endif
+#endif //WITH_EDITORONLY_DATA
 	Super::PostInitProperties();
 }
 
@@ -75,7 +75,7 @@ void UNeuralNetworkLegacy::Serialize(FArchive& Archive)
 	{
 		ReimportAssetFromEditorData();
 	}
-#endif // WITH_EDITORONLY_DATA
+#endif //WITH_EDITORONLY_DATA
 	Super::Serialize(Archive);
 }
 
@@ -95,12 +95,10 @@ void UNeuralNetworkLegacy::ReimportAssetFromEditorData()
 	}
 }
 
-#if WITH_EDITOR
 UAssetImportData* UNeuralNetworkLegacy::GetAssetImportData() const
 {
 	return AssetImportData;
 }
-#endif // WITH_EDITOR
 
 UAssetImportData* UNeuralNetworkLegacy::GetAndMaybeCreateAssetImportData()
 {
@@ -111,9 +109,7 @@ UAssetImportData* UNeuralNetworkLegacy::GetAndMaybeCreateAssetImportData()
 	}
 	return AssetImportData;
 }
-#endif // WITH_EDITOR
 
-#if WITH_EDITOR
 bool UNeuralNetworkLegacy::Load(const FString& InModelFilePath)
 {
 	DECLARE_SCOPE_CYCLE_COUNTER(TEXT("UNeuralNetworkLegacy_Load_File"), STAT_UNeuralNetworkLegacy_Load_File, STATGROUP_MachineLearning);
diff --git a/Engine/Plugins/Experimental/NeuralNetworkInference/Source/NeuralNetworkInferenceQA/NeuralNetworkInferenceQA.Build.cs b/Engine/Plugins/Experimental/NeuralNetworkInference/Source/NeuralNetworkInferenceQA/NeuralNetworkInferenceQA.Build.cs
--- a/Engine/Plugins/Experimental/NeuralNetworkInference/Source/NeuralNetworkInferenceQA/NeuralNetworkInferenceQA.Build.cs
+++ b/Engine/Plugins/Experimental/NeuralNetworkInference/Source/NeuralNetworkInferenceQA/NeuralNetworkInferenceQA.Build.cs
@@ -33,19 +33,15 @@ public class NeuralNetworkInferenceQA : ModuleRules
 			}
 		);
 
-		// Win64-only
-		if (Target.Platform == UnrealTargetPlatform.Win64)
+		// Editor-only
+		if (Target.bBuildEditor)
 		{
-			// Editor-only (to be deprecated)
-			if (Target.bBuildEditor)
-			{
-				PrivateDependencyModuleNames.AddRange
-					(
-					new string[] {
-						"ModelProtoFileReader"
-					}
-				);
-			}
+			PrivateDependencyModuleNames.AddRange
+				(
+				new string[] {
+					"ModelProtoFileReader"
+				}
+			);
 		}
 	}
 }
