NeuralNetworkInference (NNI) plugin:
- Improved ModelUnitTester error message display.

#rnx

diff --git a/Engine/Plugins/Experimental/NeuralNetworkInference/Source/NeuralNetworkInferenceQA/Private/ModelUnitTester.cpp b/Engine/Plugins/Experimental/NeuralNetworkInference/Source/NeuralNetworkInferenceQA/Private/ModelUnitTester.cpp
--- a/Engine/Plugins/Experimental/NeuralNetworkInference/Source/NeuralNetworkInferenceQA/Private/ModelUnitTester.cpp
+++ b/Engine/Plugins/Experimental/NeuralNetworkInference/Source/NeuralNetworkInferenceQA/Private/ModelUnitTester.cpp
@@ -191,34 +191,29 @@ void FModelUnitTester::ModelAccuracyTest(UNeuralNetwork* InOutNetwork, const TAr
 		const double FastCPUGPUAvgL1NormDiff = FMath::Abs((CPUAvgL1Norm - GPUAvgL1Norm)) * RelativeCoefficient * 1e6;
 		const double FastCPUAvgL1NormDiff = FMath::Abs(CPUAvgL1Norm - CPUGroundTruth) / FMath::Max(1., FMath::Min(CPUAvgL1Norm, CPUGroundTruth)) * 1e7;
 		const double FastGPUAvgL1NormDiff = FMath::Abs(GPUAvgL1Norm - GPUGroundTruth) / FMath::Max(1., FMath::Min(GPUAvgL1Norm, GPUGroundTruth)) * 1e7;
-		UE_LOG(LogNeuralNetworkInferenceQA, Display, TEXT("InputNorm = %f, CPUAvgL1Norm = %f, GPUAvgL1Norm = %f,"
+		UE_LOG(LogNeuralNetworkInferenceQA, Display, TEXT("InputNorm = %f, OutputNormCPU = %f, OutputNormGT = %f, CPUAvgL1Norm = %f, GPUAvgL1Norm = %f,"
 			"\tCPUGPUAvgL1NormDiff = %fe-3, FastCPUGPUAvgL1NormDiff = %fe-6, FastCPUAvgL1NormDiff = %fe-7, FastGPUAvgL1NormDiff = %fe-7 (1e-7 is roughly the precision for float)."),
-			GetAveragedL1Norm(InputArray), CPUAvgL1Norm, GPUAvgL1Norm, CPUGPUAvgL1NormDiff, FastCPUGPUAvgL1NormDiff, FastCPUAvgL1NormDiff, FastGPUAvgL1NormDiff);
+			GetAveragedL1Norm(InputArray), CPUAvgL1Norm, CPUGroundTruth, CPUAvgL1Norm, GPUAvgL1Norm, CPUGPUAvgL1NormDiff, FastCPUGPUAvgL1NormDiff, FastCPUAvgL1NormDiff, FastGPUAvgL1NormDiff);
 		const TArray<int64>& InputSizes = InOutNetwork->GetInputTensor().GetSizes();
 		const TArray<int64>& OutputSizes = InOutNetwork->GetOutputTensor().GetSizes();
 		const int64 MaxNumberElementsToDisplay = 100;
-		if (!FMath::IsFinite(FastCPUGPUAvgL1NormDiff) || FastCPUGPUAvgL1NormDiff > 5)
+		// Check if test failed and (if so) display information
+		bool bDidSomeTestFailed = (!FMath::IsFinite(FastCPUGPUAvgL1NormDiff) || FastCPUGPUAvgL1NormDiff > 5);
+		bDidSomeTestFailed |= (!FMath::IsFinite(CPUGPUAvgL1NormDiff) || CPUGPUAvgL1NormDiff > 1);
+		bDidSomeTestFailed |= (!FMath::IsFinite(FastCPUAvgL1NormDiff) || FastCPUAvgL1NormDiff > 30);
+		bDidSomeTestFailed |= (!FMath::IsFinite(FastGPUAvgL1NormDiff) || FastGPUAvgL1NormDiff > 30);
+		if (bDidSomeTestFailed)
 		{
-			UE_LOG(LogNeuralNetworkInferenceQA, Display, TEXT("FastCPUGPUAvgL1NormDiff (%fe-6) < 5e-6 failed."), FastCPUGPUAvgL1NormDiff);
+			UE_LOG(LogNeuralNetworkInferenceQA, Display, TEXT("FastCPUGPUAvgL1NormDiff (%fe-6) < 5e-6 might have failed."), FastCPUGPUAvgL1NormDiff);
+			UE_LOG(LogNeuralNetworkInferenceQA, Display, TEXT("CPUGPUAvgL1NormDiff (%fe-7) < 1e-7 might have failed."), CPUGPUAvgL1NormDiff);
+			UE_LOG(LogNeuralNetworkInferenceQA, Display, TEXT("FastCPUAvgL1NormDiff (%fe-7) < 30e-7 might have failed (~30 times the float precision).\nOutputArrayCPU = %s."), FastCPUAvgL1NormDiff, *FNeuralTensor(OutputArrayCPU, OutputSizes).ToString(MaxNumberElementsToDisplay));
+			UE_LOG(LogNeuralNetworkInferenceQA, Display, TEXT("FastGPUAvgL1NormDiff (%fe-7) < 30e-7 might have failed (~30 times the float precision).\nOutputArrayGPU = %s."), FastGPUAvgL1NormDiff, *FNeuralTensor(OutputArrayGPU, OutputSizes).ToString(MaxNumberElementsToDisplay));
 			UE_LOG(LogNeuralNetworkInferenceQA, Display, TEXT("Input = %s"),
 				*FNeuralTensor(InOutNetwork->GetInputTensor().GetArrayCopy<float>(), InputSizes).ToString(MaxNumberElementsToDisplay));
 			UE_LOG(LogNeuralNetworkInferenceQA, Display, TEXT("OutputArrayCPU = %s"), *FNeuralTensor(OutputArrayCPU, OutputSizes).ToString(MaxNumberElementsToDisplay));
 			UE_LOG(LogNeuralNetworkInferenceQA, Display, TEXT("OutputArrayGPU = %s"), *FNeuralTensor(OutputArrayGPU, OutputSizes).ToString(MaxNumberElementsToDisplay));
-			ensureMsgf(false, TEXT("FastCPUGPUAvgL1NormDiff < 5e-6 failed."));
+			ensureMsgf(false, TEXT("At least 1 of the 4 CPU/GPU tests failed."));
 		}
-		if (!FMath::IsFinite(CPUGPUAvgL1NormDiff) || CPUGPUAvgL1NormDiff > 1)
-		{
-			UE_LOG(LogNeuralNetworkInferenceQA, Display, TEXT("CPUGPUAvgL1NormDiff (%fe-7) < 1e-7 failed."), CPUGPUAvgL1NormDiff);
-			UE_LOG(LogNeuralNetworkInferenceQA, Display, TEXT("Input = %s"),
-				*FNeuralTensor(InOutNetwork->GetInputTensor().GetArrayCopy<float>(), InputSizes).ToString(MaxNumberElementsToDisplay));
-			UE_LOG(LogNeuralNetworkInferenceQA, Display, TEXT("OutputArrayCPU = %s"), *FNeuralTensor(OutputArrayCPU, OutputSizes).ToString(MaxNumberElementsToDisplay));
-			UE_LOG(LogNeuralNetworkInferenceQA, Display, TEXT("OutputArrayGPU = %s"), *FNeuralTensor(OutputArrayGPU, OutputSizes).ToString(MaxNumberElementsToDisplay));
-			ensureMsgf(false, TEXT("CPUGPUAvgL1NormDiff < 1e-7 failed."));
-		}
-		ensureMsgf(FMath::IsFinite(FastCPUAvgL1NormDiff) && FastCPUAvgL1NormDiff < 30, TEXT("FastCPUAvgL1NormDiff (%fe-7) < 30e-7 failed (~30 times the float precision).\nOutputArrayCPU = %s"),
-			FastCPUAvgL1NormDiff, *FNeuralTensor(OutputArrayCPU, OutputSizes).ToString(MaxNumberElementsToDisplay));
-		ensureMsgf(FMath::IsFinite(FastGPUAvgL1NormDiff) && FastGPUAvgL1NormDiff < 30, TEXT("FastGPUAvgL1NormDiff (%fe-7) < 30e-7 failed (~30 times the float precision).\nOutputArrayGPU = %s"),
-			FastGPUAvgL1NormDiff, *FNeuralTensor(OutputArrayGPU, OutputSizes).ToString(MaxNumberElementsToDisplay));
 	}
 }
 
