NeuralNetworkInference (NNI) plugin:
- Removed unused include.
- FLegacyONNXTester enabled outside of Windows (but for Editor).

#rnx

diff --git a/Engine/Plugins/Experimental/NeuralNetworkInference/Source/NeuralNetworkInference/Private/GraphProtoToNeuralNetworkConverter.cpp b/Engine/Plugins/Experimental/NeuralNetworkInference/Source/NeuralNetworkInference/Private/GraphProtoToNeuralNetworkConverter.cpp
--- a/Engine/Plugins/Experimental/NeuralNetworkInference/Source/NeuralNetworkInference/Private/GraphProtoToNeuralNetworkConverter.cpp
+++ b/Engine/Plugins/Experimental/NeuralNetworkInference/Source/NeuralNetworkInference/Private/GraphProtoToNeuralNetworkConverter.cpp
@@ -7,10 +7,6 @@
 #include "NeuralOperators.h"
 #include "NeuralTensorManager.h"
 
-#if WITH_EDITOR
-#include "ModelProtoFileReader.h"
-#endif //WITH_EDITOR
-
 
 
 /* FGraphProtoToNeuralNetworkConverter public functions
diff --git a/Engine/Plugins/Experimental/NeuralNetworkInference/Source/NeuralNetworkInferenceQA/Private/LegacyONNXTester.cpp b/Engine/Plugins/Experimental/NeuralNetworkInference/Source/NeuralNetworkInferenceQA/Private/LegacyONNXTester.cpp
--- a/Engine/Plugins/Experimental/NeuralNetworkInference/Source/NeuralNetworkInferenceQA/Private/LegacyONNXTester.cpp
+++ b/Engine/Plugins/Experimental/NeuralNetworkInference/Source/NeuralNetworkInferenceQA/Private/LegacyONNXTester.cpp
@@ -4,9 +4,7 @@
 #include "NeuralNetworkInferenceQAUtils.h"
 
 #if WITH_EDITOR
-#ifdef PLATFORM_WIN64
 #include "ModelProtoFileReader.h"
-#endif //PLATFORM_WIN64
 #endif //WITH_EDITOR
 
 
@@ -19,14 +17,10 @@ void FLegacyONNXTester::ONNXReadNetworkTest(const FString& InONNXModelFileName)
 	UE_LOG(LogNeuralNetworkInferenceQA, Display, TEXT("-------------------- Read ONNX Network And Test"));
 
 #if WITH_EDITOR
-#ifdef PLATFORM_WIN64
 	FModelProto ModelProto;
 	FModelProtoFileReader::ReadModelProtoFromFile(ModelProto, InONNXModelFileName);
 	UE_LOG(LogNeuralNetworkInferenceQA, Display, TEXT("ModelProto:\n%s"), *ModelProto.ToString());
 	ensureMsgf(ModelProto.IsLoaded(), TEXT("FLegacyONNXTester::ONNXReadNetworkTest() failed, FModelProto could not be read from InONNXModelFileName: %s."), *InONNXModelFileName);
-#else //PLATFORM_WIN64
-	UE_LOG(LogNeuralNetworkInferenceQA, Display, TEXT("ONNXReadNetworkTest test skipped (only in Windows)."));
-#endif //PLATFORM_WIN64
 #else //WITH_EDITOR
 	UE_LOG(LogNeuralNetworkInferenceQA, Display, TEXT("ONNXReadNetworkTest test skipped (only in Editor)."));
 #endif //WITH_EDITOR
