NeuralNetworkInference (NNI) plugin:
- Spliting ModelProtoFileReader to (HOPEFULLY!) avoid Linux Editor segmentation fault.

#rnx

diff --git a/Engine/Plugins/Experimental/NeuralNetworkInference/NeuralNetworkInference.uplugin b/Engine/Plugins/Experimental/NeuralNetworkInference/NeuralNetworkInference.uplugin
--- a/Engine/Plugins/Experimental/NeuralNetworkInference/NeuralNetworkInference.uplugin
+++ b/Engine/Plugins/Experimental/NeuralNetworkInference/NeuralNetworkInference.uplugin
@@ -68,6 +68,11 @@
 			"Name": "ModelProtoFileReader",
 			"Type": "Editor",
 			"LoadingPhase": "Default"
+		},
+		{
+			"Name": "ModelProtoConverter",
+			"Type": "Editor",
+			"LoadingPhase": "Default"
 		}
 	]
 }
\ No newline at end of file
diff --git a/Engine/Plugins/Experimental/NeuralNetworkInference/Source/Private/ModelProtoConverter/Generate.bat b/Engine/Plugins/Experimental/NeuralNetworkInference/Source/Private/ModelProtoConverter/Generate.bat
new file mode 100755
--- /dev/null
+++ b/Engine/Plugins/Experimental/NeuralNetworkInference/Source/Private/ModelProtoConverter/Generate.bat
@@ -0,0 +1,31 @@
+@echo off
+setlocal
+
+SET vcpkg_root=%~dp0..\..\..\..\..\..\Source\ThirdParty\vcpkg\Win64\x64-windows-static-md-v142
+SET protoc=%vcpkg_root%\protobuf_x64-windows-static-md-v142\tools\protobuf\protoc.exe
+SET include=%vcpkg_root%\protobuf_x64-windows-static-md-v142\include
+
+pushd %~dp0
+mkdir .\Generated
+
+for /R ".\Schemas" %%f in (*.proto*) do (
+	%protoc% -I=%include% --proto_path=%CD%\Schemas --cpp_out=%CD%\Generated %%f
+)
+
+for /R ".\Generated" %%f in (*.pb.cc) do (
+	type SchemaAutogenHeader.txt >%%f.new
+	type SchemaHeader.txt >>%%f.new
+	type %%f >>%%f.new
+	type SchemaFooter.txt >>%%f.new
+
+	move /y %%f.new %%f >NUL
+)
+
+for /R ".\Generated" %%f in (*.pb.h) do (
+	type SchemaAutogenHeader.txt >%%f.new
+	type %%f >>%%f.new
+
+	move /y %%f.new %%f >NUL
+)
+
+popd
diff --git a/Engine/Plugins/Experimental/NeuralNetworkInference/Source/Private/ModelProtoConverter/Internal/ModelProtoConverter.h b/Engine/Plugins/Experimental/NeuralNetworkInference/Source/Private/ModelProtoConverter/Internal/ModelProtoConverter.h
new file mode 100644
--- /dev/null
+++ b/Engine/Plugins/Experimental/NeuralNetworkInference/Source/Private/ModelProtoConverter/Internal/ModelProtoConverter.h
@@ -0,0 +1,17 @@
+// Copyright Epic Games, Inc. All Rights Reserved.
+
+#pragma once
+
+#include "CoreMinimal.h"
+#include "ModelProto.h"
+#include <istream>
+
+class MODELPROTOCONVERTER_API FModelProtoConverter
+{
+public:
+	/**
+	 * It creates and fills OutModelProto from the *.onnx file read into InIfstream.
+	 * @return False if it could not initialize it, true if successful.
+	 */
+	static bool ConvertFromONNXProto3Ifstream(FModelProto& OutModelProto, std::istream& InIfstream);
+};
diff --git a/Engine/Plugins/Experimental/NeuralNetworkInference/Source/Private/ModelProtoConverter/ModelProtoConverter.Build.cs b/Engine/Plugins/Experimental/NeuralNetworkInference/Source/Private/ModelProtoConverter/ModelProtoConverter.Build.cs
new file mode 100644
--- /dev/null
+++ b/Engine/Plugins/Experimental/NeuralNetworkInference/Source/Private/ModelProtoConverter/ModelProtoConverter.Build.cs
@@ -0,0 +1,36 @@
+// Copyright Epic Games, Inc. All Rights Reserved.
+
+using UnrealBuildTool;
+using System.IO;
+
+public class ModelProtoConverter : ModuleRules
+{
+	public ModelProtoConverter( ReadOnlyTargetRules Target ) : base( Target )
+	{
+		PCHUsage = ModuleRules.PCHUsageMode.UseExplicitOrSharedPCHs;
+
+		PrivateIncludePaths.Add(Path.Combine(ModuleDirectory, "Generated"));
+
+		PublicIncludePaths.AddRange(
+			new string[] {
+				System.IO.Path.Combine(ModuleDirectory, "..")
+			}
+		);
+
+		PublicDependencyModuleNames.AddRange
+			(
+			new string[] {
+				"Core",
+				"ModelProto"
+			}
+        );
+
+		PrivateDependencyModuleNames.AddRange
+			(
+			new string[] {
+				"Protobuf",
+                "ThirdPartyHelperAndDLLLoader"
+            }
+		);
+	}
+}
\ No newline at end of file
diff --git a/Engine/Plugins/Experimental/NeuralNetworkInference/Source/Private/ModelProtoConverter/Private/ModelProtoConverter.cpp b/Engine/Plugins/Experimental/NeuralNetworkInference/Source/Private/ModelProtoConverter/Private/ModelProtoConverter.cpp
new file mode 100644
--- /dev/null
+++ b/Engine/Plugins/Experimental/NeuralNetworkInference/Source/Private/ModelProtoConverter/Private/ModelProtoConverter.cpp
@@ -0,0 +1,507 @@
+// Copyright Epic Games, Inc. All Rights Reserved.
+
+#include "ModelProtoConverter.h"
+#include "ModelProtoConverterUtils.h"
+
+// Protobuf includes
+#ifdef WITH_PROTOBUF
+#include "ThirdPartyWarningDisabler.h"
+NNI_THIRD_PARTY_INCLUDES_START
+#undef TEXT
+#undef check
+#include "onnx.proto3.pb.h"
+NNI_THIRD_PARTY_INCLUDES_END
+#endif //WITH_PROTOBUF
+
+
+
+/* FPrivateModelProtoConverter public functions
+*****************************************************************************/
+
+#ifdef WITH_PROTOBUF
+
+class FPrivateModelProtoConverter
+{
+public:
+	static bool ConvertProto3ToUAsset(FModelProto& OutModelProto, const onnx::ModelProto& InONNXModelProto);
+
+	static bool ConvertProto3ToUAsset(FOperatorSetIdProto& OutOperatorSetIdProto, const onnx::OperatorSetIdProto& InONNXOperatorSetIdProto);
+
+	static bool ConvertProto3ToUAsset(FTrainingInfoProto& OutTrainingInfoProto, const onnx::TrainingInfoProto& InONNXTrainingInfoProto);
+
+	static bool ConvertProto3ToUAsset(FGraphProto& OutGraphProto, const onnx::GraphProto& InONNXGraphProto);
+
+	static bool ConvertProto3ToUAsset(FNodeProto& OutNodeProto, const onnx::NodeProto& InONNXNodeProto);
+
+	static bool ConvertProto3ToUAsset(FTensorAnnotation& OutTensorAnnotation, const onnx::TensorAnnotation& InONNXTensorAnnotation);
+
+	static bool ConvertProto3ToUAsset(FValueInfoProto& OutValueInfoProto, const onnx::ValueInfoProto& InONNXValueInfoProto);
+
+	static bool ConvertProto3ToUAsset(FAttributeProto& OutAttributeProto, const onnx::AttributeProto& InONNXAttributeProto);
+
+	static bool ConvertProto3ToUAsset(FTypeProto& OutTypeProto, const onnx::TypeProto& InONNXTypeProto);
+
+	static bool ConvertProto3ToUAsset(FSparseTensorProto& OutSparseTensorProto, const onnx::SparseTensorProto& InONNXSparseTensorProto);
+
+	static bool ConvertProto3ToUAsset(FTypeProtoTensor& OutTypeProtoTensor, const onnx::TypeProto_Tensor& InONNXTypeProtoTensor);
+
+	static bool ConvertProto3ToUAsset(FTensorProto& OutTensorProto, const onnx::TensorProto& InONNXTensorProto);
+
+	static bool ConvertProto3ToUAsset(FTensorShapeProto& OutTensorShapeProto, const onnx::TensorShapeProto& InONNXTensorShapeProto);
+
+	static bool ConvertProto3ToUAsset(FStringStringEntryProto& OutStringStringEntryProto, const onnx::StringStringEntryProto& FStringStringEntryProto);
+
+	static bool ConvertProto3ToUAsset(FTensorShapeProtoDimension& OutTensorShapeProtoDimension, const onnx::TensorShapeProto_Dimension& InONNXTensorShapeProtoDimension);
+
+	static bool ConvertProto3ToUAsset(FTensorProtoSegment& OutTensorProtoSegment, const onnx::TensorProto_Segment& InONNXTensorProtoSegment);
+
+	static bool ConvertProto3ToUAssetFString(TArray<FString>& OutFStringArray, const google::protobuf::RepeatedPtrField<std::basic_string<char, std::char_traits<char>, std::allocator<char>>>& InONNXFStringArray);
+
+	static bool ConvertProto3ToUAssetUInt8(TArray<uint8>& OutDataArray, const std::string& InRawDataString);
+
+	template <typename T1, typename T2>
+	static bool ConvertProto3ToUAssetBasicType(TArray<T1>& OutBasicTypeArray, const google::protobuf::RepeatedField<T2>& InONNXBasicTypeArray);
+
+	template <typename T1, typename T2>
+	static bool ConvertProto3ToUAssetProtoArrays(TArray<T1>& OutProtoArray, const google::protobuf::RepeatedPtrField<T2>& InONNXArray);
+};
+
+bool FPrivateModelProtoConverter::ConvertProto3ToUAsset(FModelProto& OutModelProto, const onnx::ModelProto& InONNXModelProto)
+{
+	// FModelProto from Protobuf onnx::ModelProto
+	OutModelProto = FModelProto();
+	OutModelProto.bIsLoaded = true;
+	OutModelProto.IRVersion = InONNXModelProto.ir_version();
+	// OperatorSetIds
+	if (!ConvertProto3ToUAssetProtoArrays(OutModelProto.OpsetImport, InONNXModelProto.opset_import()))
+	{
+		UE_LOG(LogModelProtoConverter, Warning, TEXT("FPrivateModelProtoConverter::ConvertProto3ToUAsset(): ConvertProto3ToUAsset(OperatorSetId) failed."));
+		return false;
+	}
+	OutModelProto.ProducerName = UTF8_TO_TCHAR(InONNXModelProto.producer_name().c_str());
+	OutModelProto.ProducerVersion = UTF8_TO_TCHAR(InONNXModelProto.producer_version().c_str());
+	OutModelProto.Domain = UTF8_TO_TCHAR(InONNXModelProto.domain().c_str());
+	OutModelProto.ModelVersion = InONNXModelProto.model_version();
+	OutModelProto.DocString = UTF8_TO_TCHAR(InONNXModelProto.doc_string().c_str());
+	// Graph
+	if (!ConvertProto3ToUAsset(OutModelProto.Graph, InONNXModelProto.graph()))
+	{
+		UE_LOG(LogModelProtoConverter, Warning, TEXT("FPrivateModelProtoConverter::ConvertProto3ToUAsset(): ConvertProto3GraphToUAsset(Graph) failed."));
+		return false;
+	}
+	// StringStringEntrys
+	if (!ConvertProto3ToUAssetProtoArrays(OutModelProto.MetadataProps, InONNXModelProto.metadata_props()))
+	{
+		UE_LOG(LogModelProtoConverter, Warning, TEXT("FPrivateModelProtoConverter::ConvertProto3ToUAsset(): ConvertProto3ToUAsset(StringStringEntry) failed."));
+		return false;
+	}
+	// TrainingInfos
+	if (!ConvertProto3ToUAssetProtoArrays(OutModelProto.TrainingInfo, InONNXModelProto.training_info()))
+	{
+		UE_LOG(LogModelProtoConverter, Warning, TEXT("FPrivateModelProtoConverter::ConvertProto3ToUAsset(): ConvertProto3ToUAsset(TrainingInfo) failed."));
+		return false;
+	}
+
+	// Return
+	return true;
+}
+
+bool FPrivateModelProtoConverter::ConvertProto3ToUAsset(FOperatorSetIdProto& OutOperatorSetIdProto, const onnx::OperatorSetIdProto& InONNXOperatorSetIdProto)
+{
+	OutOperatorSetIdProto = FOperatorSetIdProto();
+
+	OutOperatorSetIdProto.Domain = UTF8_TO_TCHAR(InONNXOperatorSetIdProto.domain().c_str());
+	OutOperatorSetIdProto.Version = InONNXOperatorSetIdProto.version();
+
+	// Return
+	return true;
+}
+
+bool FPrivateModelProtoConverter::ConvertProto3ToUAsset(FTrainingInfoProto& OutTrainingInfoProto, const onnx::TrainingInfoProto& InONNXTrainingInfoProto)
+{
+	OutTrainingInfoProto = FTrainingInfoProto();
+
+	OutTrainingInfoProto.bIsLoaded = true;
+	// Graph
+	if (!ConvertProto3ToUAsset(OutTrainingInfoProto.Initialization, InONNXTrainingInfoProto.initialization()) || !ConvertProto3ToUAsset(OutTrainingInfoProto.Algorithm, InONNXTrainingInfoProto.algorithm()))
+	{
+		UE_LOG(LogModelProtoConverter, Warning, TEXT("FPrivateModelProtoConverter::ConvertProto3ToUAsset(): ConvertProto3GraphToUAsset(Graph) failed."));
+		return false;
+	}
+	// StringStringEntrys
+	if (!ConvertProto3ToUAssetProtoArrays(OutTrainingInfoProto.InitializationBinding, InONNXTrainingInfoProto.initialization_binding()) || !ConvertProto3ToUAssetProtoArrays(OutTrainingInfoProto.UpdateBinding, InONNXTrainingInfoProto.update_binding()))
+	{
+		UE_LOG(LogModelProtoConverter, Warning, TEXT("FPrivateModelProtoConverter::ConvertProto3ToUAsset(): ConvertProto3ToUAsset(StringStringEntry) failed."));
+		return false;
+	}
+
+	// Return
+	return true;
+}
+
+bool FPrivateModelProtoConverter::ConvertProto3ToUAsset(FGraphProto& OutGraphProto, const onnx::GraphProto& InONNXGraphProto)
+{
+	OutGraphProto = FGraphProto();
+	OutGraphProto.bIsLoaded = true;
+	// Nodes
+	if (!ConvertProto3ToUAssetProtoArrays(OutGraphProto.Node, InONNXGraphProto.node()))
+	{
+		UE_LOG(LogModelProtoConverter, Warning, TEXT("FPrivateModelProtoConverter::ConvertProto3ToUAsset(): ConvertProto3ToUAsset(Node) failed."));
+		return false;
+	}
+	OutGraphProto.Name = UTF8_TO_TCHAR(InONNXGraphProto.name().c_str());
+	// TensorProtos
+	if (!ConvertProto3ToUAssetProtoArrays(OutGraphProto.Initializer, InONNXGraphProto.initializer()))
+	{
+		UE_LOG(LogModelProtoConverter, Warning, TEXT("FPrivateModelProtoConverter::ConvertProto3ToUAsset(): ConvertProto3ToUAsset(Tensor) failed."));
+		return false;
+	}
+	// SparseTensors
+	if (!ConvertProto3ToUAssetProtoArrays(OutGraphProto.SparseInitializer, InONNXGraphProto.sparse_initializer()))
+	{
+		UE_LOG(LogModelProtoConverter, Warning, TEXT("FPrivateModelProtoConverter::ConvertProto3ToUAsset(): ConvertProto3ToUAsset(SparseTensor) failed."));
+		return false;
+	}
+	OutGraphProto.DocString = UTF8_TO_TCHAR(InONNXGraphProto.doc_string().c_str());
+	// ValueInfos
+	if (!ConvertProto3ToUAssetProtoArrays(OutGraphProto.Input, InONNXGraphProto.input()) || !ConvertProto3ToUAssetProtoArrays(OutGraphProto.Output, InONNXGraphProto.output()) || !ConvertProto3ToUAssetProtoArrays(OutGraphProto.ValueInfo, InONNXGraphProto.value_info()))
+	{
+		UE_LOG(LogModelProtoConverter, Warning, TEXT("FPrivateModelProtoConverter::ConvertProto3ToUAsset(): ConvertProto3ToUAsset(ValueInfo) failed."));
+		return false;
+	}
+	// TensorAnnotations
+	if (!ConvertProto3ToUAssetProtoArrays(OutGraphProto.QuantizationAnnotation, InONNXGraphProto.quantization_annotation()))
+	{
+		UE_LOG(LogModelProtoConverter, Warning, TEXT("FPrivateModelProtoConverter::ConvertProto3ToUAsset(): ConvertProto3ToUAsset(TensorAnnotation) failed."));
+		return false;
+	}
+
+	if (false /* If something went wrong at any point, this is a template warning message*/)
+	{
+		UE_LOG(LogModelProtoConverter, Warning, TEXT("FPrivateModelProtoConverter::ConvertProto3ToUAsset(): FGraphProto could not be read from onnx::GraphProto."));
+		return false;
+	}
+
+	return true;
+}
+
+bool FPrivateModelProtoConverter::ConvertProto3ToUAsset(FNodeProto& OutNodeProto, const onnx::NodeProto& InONNXNodeProto)
+{
+	OutNodeProto = FNodeProto();
+
+	//FString
+	if (!ConvertProto3ToUAssetFString(OutNodeProto.Input, InONNXNodeProto.input()) || !ConvertProto3ToUAssetFString(OutNodeProto.Output, InONNXNodeProto.output()))
+	{
+		UE_LOG(LogModelProtoConverter, Warning, TEXT("FPrivateModelProtoConverter::ConvertProto3ToUAsset(): ConvertProto3ToUAssetFString() failed."));
+		return false;
+	}
+	OutNodeProto.Name = UTF8_TO_TCHAR(InONNXNodeProto.name().c_str());
+	OutNodeProto.OperatorType = UTF8_TO_TCHAR(InONNXNodeProto.op_type().c_str());
+	OutNodeProto.Domain = UTF8_TO_TCHAR(InONNXNodeProto.domain().c_str());
+	// Attributes
+	if (!ConvertProto3ToUAssetProtoArrays(OutNodeProto.Attribute, InONNXNodeProto.attribute()))
+	{
+		UE_LOG(LogModelProtoConverter, Warning, TEXT("FPrivateModelProtoConverter::ConvertProto3ToUAsset(): ConvertProto3ToUAsset(Attribute) failed."));
+		return false;
+	}
+	OutNodeProto.DocString = UTF8_TO_TCHAR(InONNXNodeProto.doc_string().c_str());
+
+	//Return
+	return true;
+}
+
+bool FPrivateModelProtoConverter::ConvertProto3ToUAsset(FTensorAnnotation& OutTensorAnnotation, const onnx::TensorAnnotation& InONNXTensorAnnotation)
+{
+	OutTensorAnnotation = FTensorAnnotation();
+
+	OutTensorAnnotation.TensorName = UTF8_TO_TCHAR(InONNXTensorAnnotation.tensor_name().c_str());
+	// StringStringEntrys
+	if (!ConvertProto3ToUAssetProtoArrays(OutTensorAnnotation.QuantParameterTensorNames, InONNXTensorAnnotation.quant_parameter_tensor_names()))
+	{
+		UE_LOG(LogModelProtoConverter, Warning, TEXT("FPrivateModelProtoConverter::ConvertProto3ToUAsset(): ConvertProto3ToUAsset(StringStringEntry) failed."));
+		return false;
+	}
+
+	//Return
+	return true;
+}
+
+bool FPrivateModelProtoConverter::ConvertProto3ToUAsset(FValueInfoProto& OutValueInfoProto, const onnx::ValueInfoProto& InONNXValueInfoProto)
+{
+	OutValueInfoProto = FValueInfoProto();
+
+	OutValueInfoProto.Name = UTF8_TO_TCHAR(InONNXValueInfoProto.name().c_str());
+	// Type
+	if (!ConvertProto3ToUAsset(OutValueInfoProto.Type, InONNXValueInfoProto.type()))
+	{
+		UE_LOG(LogModelProtoConverter, Warning, TEXT("FPrivateModelProtoConverter::ConvertProto3ToUAsset(): ConvertProto3ToUAsset(Type) failed."));
+		return false;
+	}
+	OutValueInfoProto.DocString = UTF8_TO_TCHAR(InONNXValueInfoProto.doc_string().c_str());
+
+	//Return
+	return true;
+}
+
+bool FPrivateModelProtoConverter::ConvertProto3ToUAsset(FAttributeProto& OutAttributeProto, const onnx::AttributeProto& InONNXAttributeProto)
+{
+	OutAttributeProto = FAttributeProto();
+
+	OutAttributeProto.Name = UTF8_TO_TCHAR(InONNXAttributeProto.name().c_str());
+	OutAttributeProto.DocString = UTF8_TO_TCHAR(InONNXAttributeProto.doc_string().c_str());
+	OutAttributeProto.Type = (EAttributeProtoAttributeType)InONNXAttributeProto.type();
+	OutAttributeProto.F = InONNXAttributeProto.f();
+	OutAttributeProto.I = InONNXAttributeProto.i();
+	OutAttributeProto.S = UTF8_TO_TCHAR(InONNXAttributeProto.s().c_str());
+	// TensorProtos
+	if (!ConvertProto3ToUAsset(OutAttributeProto.T, InONNXAttributeProto.t()) || !ConvertProto3ToUAssetProtoArrays(OutAttributeProto.Tensors, InONNXAttributeProto.tensors()))
+	{
+		UE_LOG(LogModelProtoConverter, Warning, TEXT("FPrivateModelProtoConverter::ConvertProto3ToUAsset(): ConvertProto3ToUAsset(Tensor) failed."));
+		return false;
+	}
+	//TSharedPtr<FGraphProto> G; /** Used TSharedPtr<> to break the circular dependency */
+	// SparseTensors
+	if (!ConvertProto3ToUAsset(OutAttributeProto.SparseTensor, InONNXAttributeProto.sparse_tensor()) || !ConvertProto3ToUAssetProtoArrays(OutAttributeProto.SparseTensors, InONNXAttributeProto.sparse_tensors()))
+	{
+		UE_LOG(LogModelProtoConverter, Warning, TEXT("FPrivateModelProtoConverter::ConvertProto3ToUAsset(): ConvertProto3ToUAsset(SparseTensor) failed."));
+		return false;
+	}
+	// floats and int64s
+	if (!ConvertProto3ToUAssetBasicType(OutAttributeProto.Floats, InONNXAttributeProto.floats()) || !ConvertProto3ToUAssetBasicType(OutAttributeProto.Integers, InONNXAttributeProto.ints()))
+	{
+		UE_LOG(LogModelProtoConverter, Warning, TEXT("FPrivateModelProtoConverter::ConvertProto3ToUAsset(): ConvertProto3ToUAssetBasicType() failed."));
+		return false;
+	}
+	//TArray<FString>
+	if (!ConvertProto3ToUAssetFString(OutAttributeProto.Strings, InONNXAttributeProto.strings()))
+	{
+		UE_LOG(LogModelProtoConverter, Warning, TEXT("FPrivateModelProtoConverter::ConvertProto3ToUAsset(): ConvertProto3ToUAssetFString() failed."));
+		return false;
+	}
+	//TArray<TSharedPtr<FGraphProto>> Graphs; /** Used TSharedPtr<> to break the circular dependency */
+
+	//Return
+	return true;
+}
+
+bool FPrivateModelProtoConverter::ConvertProto3ToUAsset(FTypeProto& OutTypeProto, const onnx::TypeProto& InONNXTypeProto)
+{
+	OutTypeProto = FTypeProto();
+	//TypeProto
+	if (!ConvertProto3ToUAsset(OutTypeProto.TensorType, InONNXTypeProto.tensor_type()))
+	{
+		UE_LOG(LogModelProtoConverter, Warning, TEXT("FPrivateModelProtoConverter::ConvertProto3ToUAsset(): ConvertProto3ToUAsset(TypeTensor) failed."));
+		return false;
+	}
+	OutTypeProto.Denotation = UTF8_TO_TCHAR(InONNXTypeProto.denotation().c_str());
+
+	//Return
+	return true;
+}
+
+bool FPrivateModelProtoConverter::ConvertProto3ToUAsset(FSparseTensorProto& OutSparseTensorProto, const onnx::SparseTensorProto& InONNXSparseTensorProto)
+{
+	OutSparseTensorProto = FSparseTensorProto();
+	//Tensor
+	if (!ConvertProto3ToUAsset(OutSparseTensorProto.Values, InONNXSparseTensorProto.values()) || !ConvertProto3ToUAsset(OutSparseTensorProto.Indices, InONNXSparseTensorProto.indices()))
+	{
+		UE_LOG(LogModelProtoConverter, Warning, TEXT("FPrivateModelProtoConverter::ConvertProto3ToUAsset(): ConvertProto3ToUAsset(Tensor) failed."));
+		return false;
+	}
+	//TArray<int64> Dimensions
+	if (!ConvertProto3ToUAssetBasicType(OutSparseTensorProto.Dimensions, InONNXSparseTensorProto.dims()))
+	{
+		UE_LOG(LogModelProtoConverter, Warning, TEXT("FPrivateModelProtoConverter::ConvertProto3ToUAsset(): ConvertProto3ToUAssetBasicType() failed."));
+		return false;
+	}
+
+	// Return
+	return true;
+}
+
+bool FPrivateModelProtoConverter::ConvertProto3ToUAsset(FTypeProtoTensor& OutTypeProtoTensor, const onnx::TypeProto_Tensor& InONNXTypeProtoTensor)
+{
+	OutTypeProtoTensor = FTypeProtoTensor();
+	
+	OutTypeProtoTensor.ElemType = InONNXTypeProtoTensor.elem_type();
+	// TensorShape
+	if (!ConvertProto3ToUAsset(OutTypeProtoTensor.Shape, InONNXTypeProtoTensor.shape()))
+	{
+		UE_LOG(LogModelProtoConverter, Warning, TEXT("FPrivateModelProtoConverter::ConvertProto3ToUAsset(): ConvertProto3ToUAsset(TensorShape) failed."));
+		return false;
+	}
+
+	//Return
+	return true;
+}
+
+bool FPrivateModelProtoConverter::ConvertProto3ToUAsset(FTensorProto& OutTensorProto, const onnx::TensorProto& InONNXTensorProto)
+{
+	OutTensorProto = FTensorProto();
+
+	// TArray<uint8>
+	if (!ConvertProto3ToUAssetUInt8(OutTensorProto.RawData, InONNXTensorProto.raw_data()))
+	{
+		UE_LOG(LogModelProtoConverter, Warning, TEXT("FPrivateModelProtoConverter::ConvertProto3ToUAsset(): ConvertProto3ToUAssetUInt8() failed."));
+		return false;
+	}
+	//TArray<int64, float, int32, double, uint64>
+	if (!ConvertProto3ToUAssetBasicType(OutTensorProto.Dimensions, InONNXTensorProto.dims()) || !ConvertProto3ToUAssetBasicType(OutTensorProto.FloatData, InONNXTensorProto.float_data()) ||
+		!ConvertProto3ToUAssetBasicType(OutTensorProto.Int32Data, InONNXTensorProto.int32_data()) || !ConvertProto3ToUAssetBasicType(OutTensorProto.Int64Data, InONNXTensorProto.int64_data()) ||
+		!ConvertProto3ToUAssetBasicType(OutTensorProto.DoubleData, InONNXTensorProto.double_data()) || !ConvertProto3ToUAssetBasicType(OutTensorProto.UInt64Data, InONNXTensorProto.uint64_data()))
+	{
+		UE_LOG(LogModelProtoConverter, Warning, TEXT("FPrivateModelProtoConverter::ConvertProto3ToUAsset(): ConvertProto3ToUAssetBasicType() failed."));
+		return false;
+	}
+	OutTensorProto.DataType = (ETensorProtoDataType)InONNXTensorProto.data_type();
+	//FTensorProtoSegment
+	if (!ConvertProto3ToUAsset(OutTensorProto.Segment, InONNXTensorProto.segment()))
+	{
+		UE_LOG(LogModelProtoConverter, Warning, TEXT("FPrivateModelProtoConverter::ConvertProto3ToUAsset(): ConvertProto3ToUAsset(TensorProtoSegment) failed."));
+		return false;
+	}
+	//TArray<FString>
+	if (!ConvertProto3ToUAssetFString(OutTensorProto.StringData, InONNXTensorProto.string_data()))
+	{
+		UE_LOG(LogModelProtoConverter, Warning, TEXT("FPrivateModelProtoConverter::ConvertProto3ToUAsset(): ConvertProto3ToUAssetFString() failed."));
+		return false;
+	}
+	OutTensorProto.Name = UTF8_TO_TCHAR(InONNXTensorProto.name().c_str());
+	OutTensorProto.DocString = UTF8_TO_TCHAR(InONNXTensorProto.doc_string().c_str());
+	// StringStringEntryProtos
+	if (!ConvertProto3ToUAssetProtoArrays(OutTensorProto.ExternalData, InONNXTensorProto.external_data()))
+	{
+		UE_LOG(LogModelProtoConverter, Warning, TEXT("FPrivateModelProtoConverter::ConvertProto3ToUAsset(): ConvertProto3ToUAsset(StringStringEntry) failed."));
+		return false;
+	}
+	OutTensorProto.DataLocation = (ETensorProtoDataLocation)InONNXTensorProto.data_location();
+
+	// Return
+	return true;
+}
+
+bool FPrivateModelProtoConverter::ConvertProto3ToUAsset(FTensorShapeProto& OutTensorShapeProto, const onnx::TensorShapeProto& InONNXTensorShapeProto)
+{
+	OutTensorShapeProto = FTensorShapeProto();
+	// TensorShapeDimensions
+	if (!ConvertProto3ToUAssetProtoArrays(OutTensorShapeProto.Dim, InONNXTensorShapeProto.dim()))
+	{
+		UE_LOG(LogModelProtoConverter, Warning, TEXT("FPrivateModelProtoConverter::ConvertProto3ToUAsset(): ConvertProto3ToUAsset(TensorShapeProtoDimension) failed."));
+		return false;
+	}
+
+	//Return
+	return true;
+}
+
+bool FPrivateModelProtoConverter::ConvertProto3ToUAsset(FStringStringEntryProto& OutStringStringEntryProto, const onnx::StringStringEntryProto& InONNXStringStringEntryProto)
+{
+	OutStringStringEntryProto = FStringStringEntryProto();
+
+	OutStringStringEntryProto.Key = UTF8_TO_TCHAR(InONNXStringStringEntryProto.key().c_str());
+	OutStringStringEntryProto.Value = UTF8_TO_TCHAR(InONNXStringStringEntryProto.value().c_str());
+
+	// Return
+	return true;
+}
+
+bool FPrivateModelProtoConverter::ConvertProto3ToUAsset(FTensorShapeProtoDimension& OutTensorShapeProtoDimension, const onnx::TensorShapeProto_Dimension& InONNXTensorShapeProtoDimension)
+{
+	OutTensorShapeProtoDimension = FTensorShapeProtoDimension();
+
+	OutTensorShapeProtoDimension.DimValue = InONNXTensorShapeProtoDimension.dim_value();
+	OutTensorShapeProtoDimension.DimParam = UTF8_TO_TCHAR(InONNXTensorShapeProtoDimension.dim_param().c_str());
+
+	// Return
+	return true;
+}
+
+bool FPrivateModelProtoConverter::ConvertProto3ToUAsset(FTensorProtoSegment& OutTensorProtoSegment, const onnx::TensorProto_Segment& InONNXTensorProtoSegment)
+{
+	OutTensorProtoSegment = FTensorProtoSegment();
+
+	OutTensorProtoSegment.Begin = InONNXTensorProtoSegment.begin();
+	OutTensorProtoSegment.End = InONNXTensorProtoSegment.end();
+
+	// Return
+	return true;
+}
+
+bool FPrivateModelProtoConverter::ConvertProto3ToUAssetFString(TArray<FString>& OutFStringArray, const google::protobuf::RepeatedPtrField<std::basic_string<char, std::char_traits<char>, std::allocator<char>>>& InONNXFStringArray)
+{
+	const int32 ArraySize = InONNXFStringArray.size();
+	OutFStringArray.SetNum(ArraySize);
+
+	for (int32 ArrayIndex = 0; ArrayIndex < ArraySize; ArrayIndex++)
+	{
+		OutFStringArray[ArrayIndex] = UTF8_TO_TCHAR(InONNXFStringArray[ArrayIndex].c_str());
+	}
+	return true;
+}
+
+bool FPrivateModelProtoConverter::ConvertProto3ToUAssetUInt8(TArray<uint8>& OutDataArray, const std::string& InRawDataString)
+{
+	const int32 ArraySize = InRawDataString.size();
+	OutDataArray.SetNumUninitialized(ArraySize);
+
+	FMemory::Memcpy(OutDataArray.GetData(), InRawDataString.c_str(), ArraySize);
+
+	// Return
+	return true;
+}
+
+template <typename T1, typename T2>
+bool FPrivateModelProtoConverter::ConvertProto3ToUAssetBasicType(TArray<T1>& OutBasicTypeArray, const google::protobuf::RepeatedField<T2>& InONNXBasicTypeArray)
+{
+	if (sizeof(T1) != sizeof(T2))
+	{
+		UE_LOG(LogModelProtoConverter, Warning, TEXT("FPrivateModelProtoConverter::ConvertProto3ToUAssetBasicType(): sizeof(T1) == sizeof(T2) failed (%d != %d)."),
+			sizeof(T1), sizeof(T2));
+		return false;
+	}
+
+	const int32 ArraySize = InONNXBasicTypeArray.size();
+	OutBasicTypeArray.SetNumUninitialized(ArraySize);
+
+	FMemory::Memcpy(OutBasicTypeArray.GetData(), &InONNXBasicTypeArray.Get(0), sizeof(T1) * ArraySize);
+
+	// Return
+	return true;
+}
+
+template <typename T1, typename T2>
+bool FPrivateModelProtoConverter::ConvertProto3ToUAssetProtoArrays(TArray<T1>& OutProtoArray, const google::protobuf::RepeatedPtrField<T2>& InONNXArray)
+{
+	const int32 ArraySize = InONNXArray.size();
+	OutProtoArray.SetNum(ArraySize);
+
+	for (int32 ArrayIndex = 0; ArrayIndex < ArraySize; ArrayIndex++)
+	{
+		if (!ConvertProto3ToUAsset(OutProtoArray[ArrayIndex], InONNXArray[ArrayIndex]))
+		{
+			UE_LOG(LogModelProtoConverter, Warning, TEXT("FPrivateModelProtoConverter::ConvertFromONNXProto3Ifstream(): ConvertProto3ToUAsset() failed."));
+			return false;
+		}
+	}
+	return true;
+}
+
+#endif //WITH_PROTOBUF
+
+
+
+/* FModelProtoConverter public functions
+*****************************************************************************/
+
+bool FModelProtoConverter::ConvertFromONNXProto3Ifstream(FModelProto& OutModelProto, std::istream& InIfstream)
+{
+#ifdef WITH_PROTOBUF
+	// Protobuf onnx::ModelProto from Ifstream
+	onnx::ModelProto ModelProto;
+	ModelProto.ParseFromIstream(&InIfstream);
+	return FPrivateModelProtoConverter::ConvertProto3ToUAsset(OutModelProto, ModelProto);
+#else //WITH_PROTOBUF
+	UE_LOG(LogModelProtoConverter, Warning, TEXT("FModelProtoConverter::ConvertFromONNXProto3Ifstream(): WITH_PROTOBUF was not defined."));
+	return false
+#endif //WITH_PROTOBUF
+}
diff --git a/Engine/Plugins/Experimental/NeuralNetworkInference/Source/Private/ModelProtoConverter/Private/ModelProtoConverterModule.cpp b/Engine/Plugins/Experimental/NeuralNetworkInference/Source/Private/ModelProtoConverter/Private/ModelProtoConverterModule.cpp
new file mode 100644
--- /dev/null
+++ b/Engine/Plugins/Experimental/NeuralNetworkInference/Source/Private/ModelProtoConverter/Private/ModelProtoConverterModule.cpp
@@ -0,0 +1,13 @@
+// Copyright Epic Games, Inc. All Rights Reserved.
+
+#include "ModelProtoConverterModule.h"
+
+void FModelProtoConverterModule::StartupModule()
+{
+}
+
+void FModelProtoConverterModule::ShutdownModule()
+{
+}
+
+IMPLEMENT_MODULE(FModelProtoConverterModule, ModelProtoConverter);
diff --git a/Engine/Plugins/Experimental/NeuralNetworkInference/Source/Private/ModelProtoConverter/Private/ModelProtoConverterModule.h b/Engine/Plugins/Experimental/NeuralNetworkInference/Source/Private/ModelProtoConverter/Private/ModelProtoConverterModule.h
new file mode 100644
--- /dev/null
+++ b/Engine/Plugins/Experimental/NeuralNetworkInference/Source/Private/ModelProtoConverter/Private/ModelProtoConverterModule.h
@@ -0,0 +1,17 @@
+// Copyright Epic Games, Inc. All Rights Reserved.
+
+#pragma once
+
+#include "CoreMinimal.h"
+#include "Modules/ModuleManager.h"
+
+/**
+* The public interface to this module
+*/
+class FModelProtoConverterModule : public IModuleInterface
+{
+public:
+	// Begin IModuleInterface
+	virtual void StartupModule() override;
+	virtual void ShutdownModule() override;
+};
diff --git a/Engine/Plugins/Experimental/NeuralNetworkInference/Source/Private/ModelProtoConverter/Private/ModelProtoConverterUtils.h b/Engine/Plugins/Experimental/NeuralNetworkInference/Source/Private/ModelProtoConverter/Private/ModelProtoConverterUtils.h
new file mode 100644
--- /dev/null
+++ b/Engine/Plugins/Experimental/NeuralNetworkInference/Source/Private/ModelProtoConverter/Private/ModelProtoConverterUtils.h
@@ -0,0 +1,7 @@
+// Copyright Epic Games, Inc. All Rights Reserved.
+
+#pragma once
+
+#include "CoreMinimal.h"
+
+DEFINE_LOG_CATEGORY_STATIC(LogModelProtoConverter, Display, All);
diff --git a/Engine/Plugins/Experimental/NeuralNetworkInference/Source/Private/ModelProtoConverter/SchemaFooter.txt b/Engine/Plugins/Experimental/NeuralNetworkInference/Source/Private/ModelProtoConverter/SchemaFooter.txt
new file mode 100644
--- /dev/null
+++ b/Engine/Plugins/Experimental/NeuralNetworkInference/Source/Private/ModelProtoConverter/SchemaFooter.txt
@@ -0,0 +1,3 @@
+
+NNI_THIRD_PARTY_INCLUDES_END
+#endif
\ No newline at end of file
diff --git a/Engine/Plugins/Experimental/NeuralNetworkInference/Source/Private/ModelProtoConverter/SchemaHeader.txt b/Engine/Plugins/Experimental/NeuralNetworkInference/Source/Private/ModelProtoConverter/SchemaHeader.txt
new file mode 100644
--- /dev/null
+++ b/Engine/Plugins/Experimental/NeuralNetworkInference/Source/Private/ModelProtoConverter/SchemaHeader.txt
@@ -0,0 +1,10 @@
+//-V::568,590
+#ifdef WITH_PROTOBUF
+#include "CoreMinimal.h"
+#include "Misc/CoreMiscDefines.h"
+
+#include "ThirdPartyWarningDisabler.h"
+NNI_THIRD_PARTY_INCLUDES_START
+#undef TEXT
+#undef check
+
diff --git a/Engine/Plugins/Experimental/NeuralNetworkInference/Source/Private/ModelProtoFileReader/Generate.bat b/Engine/Plugins/Experimental/NeuralNetworkInference/Source/Private/ModelProtoFileReader/Generate.bat
deleted file mode 100755
diff --git a/Engine/Plugins/Experimental/NeuralNetworkInference/Source/Private/ModelProtoFileReader/ModelProtoFileReader.Build.cs b/Engine/Plugins/Experimental/NeuralNetworkInference/Source/Private/ModelProtoFileReader/ModelProtoFileReader.Build.cs
--- a/Engine/Plugins/Experimental/NeuralNetworkInference/Source/Private/ModelProtoFileReader/ModelProtoFileReader.Build.cs
+++ b/Engine/Plugins/Experimental/NeuralNetworkInference/Source/Private/ModelProtoFileReader/ModelProtoFileReader.Build.cs
@@ -9,8 +9,6 @@ public class ModelProtoFileReader : ModuleRules
 	{
 		PCHUsage = ModuleRules.PCHUsageMode.UseExplicitOrSharedPCHs;
 
-		PrivateIncludePaths.Add(Path.Combine(ModuleDirectory, "Generated"));
-
 		PublicIncludePaths.AddRange(
 			new string[] {
 				System.IO.Path.Combine(ModuleDirectory, "..")
@@ -21,8 +19,6 @@ public class ModelProtoFileReader : ModuleRules
 			(
 			new string[] {
 				"Core",
-				"CoreUObject",
-				"Engine",
 				"ModelProto"
 			}
         );
@@ -30,10 +26,7 @@ public class ModelProtoFileReader : ModuleRules
 		PrivateDependencyModuleNames.AddRange
 			(
 			new string[] {
-				"NeuralNetworkInferenceCore",
-				"Projects",
-				"Protobuf",
-				"ThirdPartyHelperAndDLLLoader"
+				"ModelProtoConverter"
 			}
 		);
 	}
diff --git a/Engine/Plugins/Experimental/NeuralNetworkInference/Source/Private/ModelProtoFileReader/Private/ModelProtoConverter.cpp b/Engine/Plugins/Experimental/NeuralNetworkInference/Source/Private/ModelProtoFileReader/Private/ModelProtoConverter.cpp
deleted file mode 100644
diff --git a/Engine/Plugins/Experimental/NeuralNetworkInference/Source/Private/ModelProtoFileReader/Private/ModelProtoConverter.h b/Engine/Plugins/Experimental/NeuralNetworkInference/Source/Private/ModelProtoFileReader/Private/ModelProtoConverter.h
deleted file mode 100644
diff --git a/Engine/Plugins/Experimental/NeuralNetworkInference/Source/Private/ModelProtoFileReader/Private/ModelProtoFileReader.cpp b/Engine/Plugins/Experimental/NeuralNetworkInference/Source/Private/ModelProtoFileReader/Private/ModelProtoFileReader.cpp
--- a/Engine/Plugins/Experimental/NeuralNetworkInference/Source/Private/ModelProtoFileReader/Private/ModelProtoFileReader.cpp
+++ b/Engine/Plugins/Experimental/NeuralNetworkInference/Source/Private/ModelProtoFileReader/Private/ModelProtoFileReader.cpp
@@ -2,7 +2,6 @@
 
 #include "ModelProtoFileReader.h"
 #include "ModelProtoFileReaderUtils.h"
-#include "Interfaces/IPluginManager.h"
 #include "Misc/FileHelper.h"
 #include "Misc/Paths.h"
 // Protobuf includes
diff --git a/Engine/Plugins/Experimental/NeuralNetworkInference/Source/Private/ModelProtoFileReader/SchemaAutogenHeader.txt b/Engine/Plugins/Experimental/NeuralNetworkInference/Source/Private/ModelProtoFileReader/SchemaAutogenHeader.txt
deleted file mode 100644
diff --git a/Engine/Plugins/Experimental/NeuralNetworkInference/Source/Private/ModelProtoFileReader/SchemaFooter.txt b/Engine/Plugins/Experimental/NeuralNetworkInference/Source/Private/ModelProtoFileReader/SchemaFooter.txt
deleted file mode 100644
diff --git a/Engine/Plugins/Experimental/NeuralNetworkInference/Source/Private/ModelProtoFileReader/SchemaHeader.txt b/Engine/Plugins/Experimental/NeuralNetworkInference/Source/Private/ModelProtoFileReader/SchemaHeader.txt
deleted file mode 100644
