NeuralNetworkInference (NNI) plugin:
- BackEnd::UEOnly (even more) partially implemented for UNN. UNN runs with UEOnly, but input/output cannot be read/filled.

#rnx

diff --git a/Engine/Plugins/Experimental/NeuralNetworkInference/Source/NeuralNetworkInference/Private/NeuralNetwork.cpp b/Engine/Plugins/Experimental/NeuralNetworkInference/Source/NeuralNetworkInference/Private/NeuralNetwork.cpp
--- a/Engine/Plugins/Experimental/NeuralNetworkInference/Source/NeuralNetworkInference/Private/NeuralNetwork.cpp
+++ b/Engine/Plugins/Experimental/NeuralNetworkInference/Source/NeuralNetworkInference/Private/NeuralNetwork.cpp
@@ -63,6 +63,8 @@ UNeuralNetwork::~UNeuralNetwork()
 #if WITH_EDITOR
 bool UNeuralNetwork::Load(const FString& InModelFilePath)
 {
+	DECLARE_SCOPE_CYCLE_COUNTER(TEXT("UNeuralNetwork_Load_EditorOnly"), STAT_UNeuralNetwork_Load, STATGROUP_MachineLearning);
+
 	// Clean previous networks
 	bIsLoaded = false;
 
@@ -99,7 +101,7 @@ bool UNeuralNetwork::Load(const FString& InModelFilePath)
 	// UEOnly
 	else if (BackEndForCurrentPlatform == ENeuralBackEnd::UEOnly)
 	{
-		bIsLoaded = UNeuralNetwork::FImplBackEndUEOnly::Load(/*ModelReadFromDiskInBytes*/ ModelFullFilePath);
+		bIsLoaded = UNeuralNetwork::FImplBackEndUEOnly::Load(ImplBackEndUEOnly, /*ModelReadFromDiskInBytes*/ ModelFullFilePath);
 	}
 	// Unknown
 	else
@@ -114,6 +116,8 @@ bool UNeuralNetwork::Load(const FString& InModelFilePath)
 
 bool UNeuralNetwork::Load()
 {
+	DECLARE_SCOPE_CYCLE_COUNTER(TEXT("UNeuralNetwork_Load"), STAT_UNeuralNetwork_Load, STATGROUP_MachineLearning);
+
 	// Clean previous networks
 	bIsLoaded = false;
 
@@ -125,7 +129,7 @@ bool UNeuralNetwork::Load()
 	// UEOnly
 	else if (BackEndForCurrentPlatform == ENeuralBackEnd::UEOnly)
 	{
-		UE_LOG(LogNeuralNetworkInference, Warning, TEXT("UNeuralNetwork::Load(): Platform or Operating System not suported yet for [BackEnd,BackEndForCurrentPlatform] = [%d,%d]."), (int32)BackEnd, (int32)BackEndForCurrentPlatform);
+		bIsLoaded = UNeuralNetwork::FImplBackEndUEOnly::Load(ImplBackEndUEOnly, /*ModelReadFromDiskInBytes*/ ModelFullFilePath);
 	}
 	// Unknown
 	else
@@ -248,6 +252,8 @@ const FNeuralTensors& UNeuralNetwork::GetOutputTensors() const
 
 void UNeuralNetwork::Run()
 {
+	DECLARE_SCOPE_CYCLE_COUNTER(TEXT("UNeuralNetwork_Run"), STAT_UNeuralNetwork_Run, STATGROUP_MachineLearning);
+
 	// Sanity checks
 	if (!bIsLoaded)
 	{
@@ -264,7 +270,7 @@ void UNeuralNetwork::Run()
 	// UEOnly
 	else if (BackEndForCurrentPlatform == ENeuralBackEnd::UEOnly)
 	{
-		UE_LOG(LogNeuralNetworkInference, Warning, TEXT("UNeuralNetwork::Run(): Platform or Operating System not suported yet for [BackEnd,BackEndForCurrentPlatform] = [%d,%d]."), (int32)BackEnd, (int32)BackEndForCurrentPlatform);
+		ImplBackEndUEOnly->Run(OnAsyncRunCompletedDelegate, SynchronousMode, DeviceType, InputDeviceType, OutputDeviceType);
 	}
 
 	// Unknown
diff --git a/Engine/Plugins/Experimental/NeuralNetworkInference/Source/NeuralNetworkInference/Private/NeuralNetworkImplBackEndUEOnly.cpp b/Engine/Plugins/Experimental/NeuralNetworkInference/Source/NeuralNetworkInference/Private/NeuralNetworkImplBackEndUEOnly.cpp
--- a/Engine/Plugins/Experimental/NeuralNetworkInference/Source/NeuralNetworkInference/Private/NeuralNetworkImplBackEndUEOnly.cpp
+++ b/Engine/Plugins/Experimental/NeuralNetworkInference/Source/NeuralNetworkInference/Private/NeuralNetworkImplBackEndUEOnly.cpp
@@ -1,18 +1,15 @@
 // Copyright Epic Games, Inc. All Rights Reserved.
 
 #include "NeuralNetworkImplBackEndUEOnly.h"
-//#include "NeuralOperator.h"
-//#include "GraphProtoToNeuralNetworkConverter.h"
-//#include "NeuralNetworkInferenceUtils.h"
-//#include "EditorFramework/AssetImportData.h"
-//#include "HAL/FileManager.h"
-//#include "RenderGraphBuilder.h"
-//#include "RenderingThread.h"
-//#include "RHI.h"
+#include "GraphProtoToNeuralNetworkConverter.h"
+#include "NeuralNetworkInferenceUtils.h"
+#include "RenderGraphBuilder.h"
+#include "RenderingThread.h"
+#include "RHI.h"
 
-//#if WITH_EDITOR
-//#include "ModelProtoFileReader.h"
-//#endif //WITH_EDITOR
+#if WITH_EDITOR
+#include "ModelProtoFileReader.h"
+#endif //WITH_EDITOR
 
 
 
@@ -21,49 +18,49 @@
 
 #if WITH_EDITOR
 
-bool UNeuralNetwork::FImplBackEndUEOnly::Load(const FString& InModelFilePath)
+bool UNeuralNetwork::FImplBackEndUEOnly::Load(TSharedPtr<FImplBackEndUEOnly>& InOutImplBackEndUEOnly, const FString& InModelFilePath)
 {
-//	DECLARE_SCOPE_CYCLE_COUNTER(TEXT("UNeuralNetworkLegacy_Load_File"), STAT_UNeuralNetworkLegacy_Load_File, STATGROUP_MachineLearning);
-//
-//	// Clean previous networks
-//	if (bIsLoaded)
-//	{
-//		UE_LOG(LogNeuralNetworkInference, Warning, TEXT("UNeuralNetwork::FImplBackEndUEOnly::Load(): A model was previously loaded, removing it and reloading the new model."));
-//		Operators.Empty();
-//		ModelProto = FModelProto();
-//		bIsLoaded = false;
-//		bAreTensorsInGpu = false;
-//	}
-//
-//	// Read ModelProto
-//	if (!FModelProtoFileReader::ReadModelProtoFromFile(ModelProto, InModelFilePath) || !ModelProto.IsLoaded())
-//	{
-//		UE_LOG(LogNeuralNetworkInference, Warning, TEXT("UNeuralNetwork::FImplBackEndUEOnly::Load(): Model could not be loaded from %s."), *InModelFilePath);
-//		return false;
-//	}
-//	// Turn ModelProto into Operators
-//	bIsLoaded = FGraphProtoToNeuralNetworkConverter::Translate(Operators, TensorManager, ModelProto.GetGraph(), InModelFilePath);
-//	if (!TensorManager.IsLoaded())
-//	{
-//		UE_LOG(LogNeuralNetworkInference, Warning, TEXT("UNeuralNetwork::FImplBackEndUEOnly::Load(): TensorManager could not be loaded from %s."), *InModelFilePath);
+	// Initialize InOutImplBackEndUEOnly
+	if (!InOutImplBackEndUEOnly.IsValid())
+	{
+		InOutImplBackEndUEOnly = MakeShared<FImplBackEndUEOnly>();
+	}
+	// Clean previous networks
+	InOutImplBackEndUEOnly->Operators.Empty();
+	InOutImplBackEndUEOnly->ModelProto = FModelProto();
+	InOutImplBackEndUEOnly->bAreTensorsInGpu = false;
+	// Read ModelProto
+	if (!FModelProtoFileReader::ReadModelProtoFromFile(InOutImplBackEndUEOnly->ModelProto, InModelFilePath) || !InOutImplBackEndUEOnly->ModelProto.IsLoaded())
+	{
+		UE_LOG(LogNeuralNetworkInference, Warning, TEXT("UNeuralNetwork::FImplBackEndUEOnly::Load(): Model could not be loaded from %s."), *InModelFilePath);
 		return false;
-//	}
-	//return true;
+	}
+	// Turn ModelProto into Operators
+	if (!FGraphProtoToNeuralNetworkConverter::Translate(InOutImplBackEndUEOnly->Operators, InOutImplBackEndUEOnly->TensorManager, InOutImplBackEndUEOnly->ModelProto.GetGraph(), InModelFilePath) || !InOutImplBackEndUEOnly->TensorManager.IsLoaded())
+	{
+		UE_LOG(LogNeuralNetworkInference, Warning, TEXT("UNeuralNetwork::FImplBackEndUEOnly::Load(): TensorManager could not be loaded from %s."), *InModelFilePath);
+		return false;
+	}
+	return true;
 }
 #endif //WITH_EDITOR
 
-bool UNeuralNetwork::FImplBackEndUEOnly::Load(FNeuralTensorManager& InTensorManager, const TArray<TSharedPtr<FNeuralOperator>>& InOperators)
-{
-	//DECLARE_SCOPE_CYCLE_COUNTER(TEXT("UNeuralNetworkLegacy_Load"), STAT_UNeuralNetworkLegacy_Load, STATGROUP_MachineLearning);
-	//if (!InTensorManager.IsLoaded())
-	//{
-	//	UE_LOG(LogNeuralNetworkInference, Warning, TEXT("UNeuralNetwork::FImplBackEndUEOnly::Load(): TensorManager could not be loaded."));
-	//}
-	//Swap(TensorManager, InTensorManager);
-	//Operators = InOperators;
-	//return (Operators.Num() > 0 && TensorManager.IsLoaded());
-return false;
-}
+//bool UNeuralNetwork::FImplBackEndUEOnly::Load(TSharedPtr<FImplBackEndUEOnly>& InOutImplBackEndUEOnly, FNeuralTensorManager& InTensorManager, const TArray<TSharedPtr<FNeuralOperator>>& InOperators)
+//{
+//	// Initialize InOutImplBackEndUEOnly
+//	if (!InOutImplBackEndUEOnly.IsValid())
+//	{
+//		InOutImplBackEndUEOnly = MakeShared<FImplBackEndUEOnly>();
+//	}
+//	// Load
+//	if (!InTensorManager.IsLoaded())
+//	{
+//		UE_LOG(LogNeuralNetworkInference, Warning, TEXT("UNeuralNetwork::FImplBackEndUEOnly::Load(): TensorManager could not be loaded."));
+//	}
+//	Swap(InOutImplBackEndUEOnly->TensorManager, InTensorManager);
+//	InOutImplBackEndUEOnly->Operators = InOperators;
+//	return (InOutImplBackEndUEOnly->Operators.Num() > 0 && InOutImplBackEndUEOnly->TensorManager.IsLoaded());
+//}
 
 //const TArray<FNeuralTensor>& UNeuralNetwork::FImplBackEndUEOnly::GetTensors() const
 //{
@@ -134,171 +131,160 @@ return false;
 //{
 //	return TensorManager.CreateOutputTensorMap();
 //}
-//
-//void UNeuralNetwork::FImplBackEndUEOnly::Run(const ENeuralNetworkSynchronousMode InSynchronousMode, const ENeuralDeviceType InInputDeviceType, const ENeuralDeviceType InOutputDeviceType, const bool bRunGPUEmptyOnlyForProfiling)
-//{
-//	DECLARE_SCOPE_CYCLE_COUNTER(TEXT("UNeuralNetworkLegacy_Run"), STAT_UNeuralNetworkLegacy_Run, STATGROUP_MachineLearning);
-//	// Sanity checks
-//	if (!bIsLoaded)
-//	{
-//		UE_LOG(LogNeuralNetworkInference, Warning, TEXT("UNeuralNetwork::FImplBackEndUEOnly::Run(): No architecture has been loaded yet. Run() will not work until IsLoaded() returns true."));
-//		return;
-//	}
-//	// Run UNeuralNetworkLegacy
-//	if (Operators.Num() > 0)
-//	{
-//		// Run graph
-//		if (DeviceType == ENeuralDeviceType::CPU)
-//		{
-//			DECLARE_SCOPE_CYCLE_COUNTER(TEXT("UNeuralNetworkLegacy_Run::Forward_CPU"), STAT_UNeuralNetworkLegacy_Run_Forward_CPU, STATGROUP_MachineLearning);
-//			// Run each operator forward pass
-//			for (TSharedPtr<FNeuralOperator>& Operator : Operators)
-//			{
-//				Operator->ForwardCPU();
-//			}
-//			// Run each operator post forward pass
-//			for (TSharedPtr<FNeuralOperator>& Operator : Operators)
-//			{
-//				Operator->PostForwardCPU();
-//			}
-//		}
-//		else if (DeviceType == ENeuralDeviceType::GPU)
-//		{
-//			DECLARE_SCOPE_CYCLE_COUNTER(TEXT("UNeuralNetworkLegacy_Run::Forward_GPU"), STAT_UNeuralNetworkLegacy_Run_Forward_GPU, STATGROUP_MachineLearning);
-//			// Sanity check
-//			if (TensorManager.GetTensorsMutable().Num() < 1)
-//			{
-//				UE_LOG(LogNeuralNetworkInference, Warning, TEXT("UNeuralNetwork::FImplBackEndUEOnly::Run(): Tensors.Num() = %d (should be > 0)."), TensorManager.GetTensorsMutable().Num());
-//				return;
-//			}
-//
-//			// On RHI thread
-//			ENQUEUE_RENDER_COMMAND(UNeuralNetworkLegacy_Run_RenderThread)(
-//				[this, bRunGPUEmptyOnlyForProfiling, InSynchronousMode, InInputDeviceType, InOutputDeviceType](FRHICommandListImmediate& RHICmdList)
-//				{
-//					FMemMark Mark(FMemStack::Get());
-//					FRDGBuilder GraphBuilder(RHICmdList, RDG_EVENT_NAME("UNeuralNetwork::FImplBackEndUEOnly::Run()"));
-//
-//					// Move memory from CPU to GPU
-//					TArray<FNeuralTensor>& Tensors = TensorManager.GetTensorsMutable();
-//					const bool bIsInputInCPU = (InInputDeviceType == ENeuralDeviceType::CPU);
-//					// Move only input tensors to GPU (once per UNeuralNetwork::FImplBackEndUEOnly::Run())
-//					if (bAreTensorsInGpu)
-//					{
-//						// If input is on CPU, move inputs to GPU
-//						if (bIsInputInCPU)
-//						{
-//							for (const int32 InputIndex : TensorManager.GetInputIndexes())
-//							{
-//								Tensors[InputIndex].ToGPU_RenderThread(&GraphBuilder);
-//							}
-//							for (const int32 NonInputIndex : TensorManager.GetNonInputIndexes())
-//							{
-//								Tensors[NonInputIndex].UpdateSRVAndOrUAV_RenderThread(&GraphBuilder);
-//							}
-//						}
-//						// If input is already on GPU, just refresh all w.r.t. GraphBuilder
-//						else
-//						{
-//							for (FNeuralTensor& Tensor : Tensors)
-//							{
-//								Tensor.UpdateSRVAndOrUAV_RenderThread(&GraphBuilder);
-//							}
-//						}
-//					}
-//					// Move all (input, intermediate, weight, output) tensors to GPU and also call Operators.ToGPU() to move their auxiliary memory into GPU
-//					// Only once per UNeuralNetworkLegacy instance, or until Load() is called again
-//					else
-//					{
-//						// If input is on CPU, move inputs to GPU
-//						if (bIsInputInCPU)
-//						{
-//							// Move tensors to GPU
-//							for (FNeuralTensor& Tensor : Tensors)
-//							{
-//								Tensor.ToGPU_RenderThread(&GraphBuilder);
-//							}
-//						}
-//						// If input is already on GPU, just refresh inputs
-//						else
-//						{
-//							for (const int32 InputIndex : TensorManager.GetInputIndexes())
-//							{
-//								Tensors[InputIndex].UpdateSRVAndOrUAV_RenderThread(&GraphBuilder);
-//							}
-//							for (const int32 NonInputIndex : TensorManager.GetNonInputIndexes())
-//							{
-//								Tensors[NonInputIndex].ToGPU_RenderThread(&GraphBuilder);
-//							}
-//						}
-//						// Operators->ToGPU()
-//						for (TSharedPtr<FNeuralOperator>& Operator : Operators)
-//						{
-//							Operator->ToGPU_RenderThread();
-//						}
-//						// bAreTensorsInGpu is true
-//						bAreTensorsInGpu = true;
-//					}
-//
-//					if (!bRunGPUEmptyOnlyForProfiling)
-//					{
-//						// Run each operator forward pass
-//						for (TSharedPtr<FNeuralOperator>& Operator : this->Operators)
-//						{
-//							Operator->ForwardGPU_RenderThread(&GraphBuilder);
-//						}
-//
-//						// Run each operator post forward pass
-//						for (TSharedPtr<FNeuralOperator>& Operator : this->Operators)
-//						{
-//							Operator->PostForwardGPU_RenderThread(&GraphBuilder);
-//						}
-//					}
-//
-//					// Move memory from GPU to CPU
-//					const bool bIsOutputInCPU = (InOutputDeviceType == ENeuralDeviceType::CPU);
-//					if (bIsOutputInCPU)
-//					{
-//						for (const int32 OutputIndex : TensorManager.GetOutputIndexes())
-//						{
-//							Tensors[OutputIndex].ToCPU_RenderThread(&GraphBuilder);
-//						}
-//					}
-//
-//					// Broadcast delegates (from the render thread)
-//					if (InSynchronousMode == ENeuralNetworkSynchronousMode::Asynchronous)
-//					{
-//						GraphBuilder.AddPass(
-//							RDG_EVENT_NAME("Async delegate broadcast"),
-//							ERDGPassFlags::None,
-//							[this](FRHICommandListImmediate& RHICmdList)
-//						{
-//							OnAsyncRunCompletedDelegate.ExecuteIfBound();
-//						});
-//					}
-//
-//					// Execute Render Graph
-//					GraphBuilder.Execute();
-//				}
-//			);
-//
-//			// Block thread until GPU has finished
-//			if (InSynchronousMode == ENeuralNetworkSynchronousMode::Synchronous)
-//			{
-//				FNeuralNetworkInferenceUtils::WaitUntilRHIFinished();
-//			}
-//		}
-//		else
-//		{
-//			UE_LOG(LogNeuralNetworkInference, Warning, TEXT("UNeuralNetwork::FImplBackEndUEOnly::Run(): Unknown DeviceType = %d."), (int32)DeviceType);
-//		}
-//	}
-//	else
-//	{
-//		UE_LOG(LogNeuralNetworkInference, Warning, TEXT("UNeuralNetwork::FImplBackEndUEOnly::Run() called with an empty model."));
-//	}
-//}
-//
+
+void UNeuralNetwork::FImplBackEndUEOnly::Run(FOnAsyncRunCompleted& InOutOnAsyncRunCompletedDelegate, const ENeuralNetworkSynchronousMode InSynchronousMode, const ENeuralDeviceType InDeviceType, const ENeuralDeviceType InInputDeviceType, const ENeuralDeviceType InOutputDeviceType)
+{
+	// Run UNeuralNetworkLegacy
+	if (Operators.Num() > 0)
+	{
+		// Run graph
+		if (InDeviceType == ENeuralDeviceType::CPU)
+		{
+			DECLARE_SCOPE_CYCLE_COUNTER(TEXT("UNeuralNetwork_UEOnly_Run::Forward_CPU"), STAT_UNeuralNetwork_UEOnly_Run_Forward_CPU, STATGROUP_MachineLearning);
+			// Run each operator forward pass
+			for (TSharedPtr<FNeuralOperator>& Operator : Operators)
+			{
+				Operator->ForwardCPU();
+			}
+			// Run each operator post forward pass
+			for (TSharedPtr<FNeuralOperator>& Operator : Operators)
+			{
+				Operator->PostForwardCPU();
+			}
+		}
+		else if (InDeviceType == ENeuralDeviceType::GPU)
+		{
+			DECLARE_SCOPE_CYCLE_COUNTER(TEXT("UNeuralNetwork_UEOnly_Run::Forward_GPU"), STAT_UNeuralNetwork_UEOnly_Run_Forward_GPU, STATGROUP_MachineLearning);
+			// Sanity check
+			if (TensorManager.GetTensorsMutable().Num() < 1)
+			{
+				UE_LOG(LogNeuralNetworkInference, Warning, TEXT("UNeuralNetwork::FImplBackEndUEOnly::Run(): Tensors.Num() = %d (should be > 0)."), TensorManager.GetTensorsMutable().Num());
+				return;
+			}
+
+			// On RHI thread
+			ENQUEUE_RENDER_COMMAND(UNeuralNetwork_UEOnly_Run_RenderThread)(
+				[this, &InOutOnAsyncRunCompletedDelegate, InSynchronousMode, InInputDeviceType, InOutputDeviceType](FRHICommandListImmediate& RHICmdList)
+				{
+					FRDGBuilder GraphBuilder(RHICmdList, RDG_EVENT_NAME("UNeuralNetwork::FImplBackEndUEOnly::Run()"));
+
+					// Move memory from CPU to GPU
+					TArray<FNeuralTensor>& Tensors = TensorManager.GetTensorsMutable();
+					const bool bIsInputInCPU = (InInputDeviceType == ENeuralDeviceType::CPU);
+					// Move only input tensors to GPU (once per UNeuralNetwork::FImplBackEndUEOnly::Run())
+					if (bAreTensorsInGpu)
+					{
+						// If input is on CPU, move inputs to GPU
+						if (bIsInputInCPU)
+						{
+							for (const int32 InputIndex : TensorManager.GetInputIndexes())
+							{
+								Tensors[InputIndex].ToGPU_RenderThread(&GraphBuilder);
+							}
+							for (const int32 NonInputIndex : TensorManager.GetNonInputIndexes())
+							{
+								Tensors[NonInputIndex].UpdateSRVAndOrUAV_RenderThread(&GraphBuilder);
+							}
+						}
+						// If input is already on GPU, just refresh all w.r.t. GraphBuilder
+						else
+						{
+							for (FNeuralTensor& Tensor : Tensors)
+							{
+								Tensor.UpdateSRVAndOrUAV_RenderThread(&GraphBuilder);
+							}
+						}
+					}
+					// Move all (input, intermediate, weight, output) tensors to GPU and also call Operators.ToGPU() to move their auxiliary memory into GPU
+					// Only once per UNeuralNetworkLegacy instance, or until Load() is called again
+					else
+					{
+						// If input is on CPU, move inputs to GPU
+						if (bIsInputInCPU)
+						{
+							// Move tensors to GPU
+							for (FNeuralTensor& Tensor : Tensors)
+							{
+								Tensor.ToGPU_RenderThread(&GraphBuilder);
+							}
+						}
+						// If input is already on GPU, just refresh inputs
+						else
+						{
+							for (const int32 InputIndex : TensorManager.GetInputIndexes())
+							{
+								Tensors[InputIndex].UpdateSRVAndOrUAV_RenderThread(&GraphBuilder);
+							}
+							for (const int32 NonInputIndex : TensorManager.GetNonInputIndexes())
+							{
+								Tensors[NonInputIndex].ToGPU_RenderThread(&GraphBuilder);
+							}
+						}
+						// Operators->ToGPU()
+						for (TSharedPtr<FNeuralOperator>& Operator : Operators)
+						{
+							Operator->ToGPU_RenderThread();
+						}
+						// bAreTensorsInGpu is true
+						bAreTensorsInGpu = true;
+					}
+
+					// Run each operator forward pass
+					for (TSharedPtr<FNeuralOperator>& Operator : this->Operators)
+					{
+						Operator->ForwardGPU_RenderThread(&GraphBuilder);
+					}
+
+					// Run each operator post forward pass
+					for (TSharedPtr<FNeuralOperator>& Operator : this->Operators)
+					{
+						Operator->PostForwardGPU_RenderThread(&GraphBuilder);
+					}
+
+					// Move memory from GPU to CPU
+					const bool bIsOutputInCPU = (InOutputDeviceType == ENeuralDeviceType::CPU);
+					if (bIsOutputInCPU)
+					{
+						for (const int32 OutputIndex : TensorManager.GetOutputIndexes())
+						{
+							Tensors[OutputIndex].ToCPU_RenderThread(&GraphBuilder);
+						}
+					}
+
+					// Broadcast delegates (from the render thread)
+					if (InSynchronousMode == ENeuralNetworkSynchronousMode::Asynchronous)
+					{
+						GraphBuilder.AddPass(
+							RDG_EVENT_NAME("Async delegate broadcast"),
+							ERDGPassFlags::None,
+							[&InOutOnAsyncRunCompletedDelegate](FRHICommandListImmediate& RHICmdList)
+						{
+							InOutOnAsyncRunCompletedDelegate.ExecuteIfBound();
+						});
+					}
+
+					// Execute Render Graph
+					GraphBuilder.Execute();
+				}
+			);
+
+			// Block thread until GPU has finished
+			if (InSynchronousMode == ENeuralNetworkSynchronousMode::Synchronous)
+			{
+				FNeuralNetworkInferenceUtils::WaitUntilRHIFinished();
+			}
+		}
+		else
+		{
+			UE_LOG(LogNeuralNetworkInference, Warning, TEXT("UNeuralNetwork::FImplBackEndUEOnly::Run(): Unknown DeviceType = %d."), (int32)InDeviceType);
+		}
+	}
+	else
+	{
+		UE_LOG(LogNeuralNetworkInference, Warning, TEXT("UNeuralNetwork::FImplBackEndUEOnly::Run() called with an empty model."));
+	}
+}
+
 //FString UNeuralNetwork::FImplBackEndUEOnly::ToString() const
 //{
 //	// Add GraphProto
diff --git a/Engine/Plugins/Experimental/NeuralNetworkInference/Source/NeuralNetworkInference/Private/NeuralNetworkImplBackEndUEOnly.h b/Engine/Plugins/Experimental/NeuralNetworkInference/Source/NeuralNetworkInference/Private/NeuralNetworkImplBackEndUEOnly.h
--- a/Engine/Plugins/Experimental/NeuralNetworkInference/Source/NeuralNetworkInference/Private/NeuralNetworkImplBackEndUEOnly.h
+++ b/Engine/Plugins/Experimental/NeuralNetworkInference/Source/NeuralNetworkInference/Private/NeuralNetworkImplBackEndUEOnly.h
@@ -32,8 +32,10 @@ struct UNeuralNetwork::FImplBackEndUEOnly
 	TArray<TSharedPtr<FNeuralOperator>> Operators;
 
  #if WITH_EDITOR
- 	static bool Load(const FString& InModelFullFilePath);
+ 	static bool Load(TSharedPtr<FImplBackEndUEOnly>& InOutImplBackEndUEOnly, const FString& InModelFullFilePath);
  #endif //WITH_EDITOR
 
-	static bool Load(FNeuralTensorManager& InTensorManager, const TArray<TSharedPtr<FNeuralOperator>>& InOperators);
+	//static bool Load(TSharedPtr<FImplBackEndUEOnly>& InOutImplBackEndUEOnly, FNeuralTensorManager& InTensorManager, const TArray<TSharedPtr<FNeuralOperator>>& InOperators);
+
+	void Run(FOnAsyncRunCompleted& InOutOnAsyncRunCompletedDelegate, const ENeuralNetworkSynchronousMode InSynchronousMode, const ENeuralDeviceType InDeviceType, const ENeuralDeviceType InInputDeviceType, const ENeuralDeviceType InOutputDeviceType);
 };
