NeuralNetworkInference (NNI) plugin:
- CIS fix for XBoxOneGDK.

#rnx

diff --git a/Engine/Plugins/Experimental/NeuralNetworkInference/Source/NeuralNetworkInference/Private/NeuralNetworkImplBackEndUEOnly.cpp b/Engine/Plugins/Experimental/NeuralNetworkInference/Source/NeuralNetworkInference/Private/NeuralNetworkImplBackEndUEOnly.cpp
--- a/Engine/Plugins/Experimental/NeuralNetworkInference/Source/NeuralNetworkInference/Private/NeuralNetworkImplBackEndUEOnly.cpp
+++ b/Engine/Plugins/Experimental/NeuralNetworkInference/Source/NeuralNetworkInference/Private/NeuralNetworkImplBackEndUEOnly.cpp
@@ -16,10 +16,9 @@
 /* FImplBackEndUEOnly public functions
  *****************************************************************************/
 
-#if WITH_EDITOR
-
 bool UNeuralNetwork::FImplBackEndUEOnly::Load(TSharedPtr<FImplBackEndUEOnly>& InOutImplBackEndUEOnly, const FString& InModelFilePath)
 {
+#if WITH_EDITOR
 	// Initialize InOutImplBackEndUEOnly
 	if (!InOutImplBackEndUEOnly.IsValid())
 	{
@@ -42,8 +41,11 @@ bool UNeuralNetwork::FImplBackEndUEOnly::Load(TSharedPtr<FImplBackEndUEOnly>& In
 		return false;
 	}
 	return true;
-}
+#else //WITH_EDITOR
+	UE_LOG(LogNeuralNetworkInference, Warning, TEXT("UNeuralNetwork::FImplBackEndUEOnly::Load(): Only implemented for Editor mode."));
+	return false;
 #endif //WITH_EDITOR
+}
 
 //bool UNeuralNetwork::FImplBackEndUEOnly::Load(TSharedPtr<FImplBackEndUEOnly>& InOutImplBackEndUEOnly, FNeuralTensorManager& InTensorManager, const TArray<TSharedPtr<FNeuralOperator>>& InOperators)
 //{
diff --git a/Engine/Plugins/Experimental/NeuralNetworkInference/Source/NeuralNetworkInference/Private/NeuralNetworkImplBackEndUEOnly.h b/Engine/Plugins/Experimental/NeuralNetworkInference/Source/NeuralNetworkInference/Private/NeuralNetworkImplBackEndUEOnly.h
--- a/Engine/Plugins/Experimental/NeuralNetworkInference/Source/NeuralNetworkInference/Private/NeuralNetworkImplBackEndUEOnly.h
+++ b/Engine/Plugins/Experimental/NeuralNetworkInference/Source/NeuralNetworkInference/Private/NeuralNetworkImplBackEndUEOnly.h
@@ -31,9 +31,7 @@ struct UNeuralNetwork::FImplBackEndUEOnly
 	 */
 	TArray<TSharedPtr<FNeuralOperator>> Operators;
 
- #if WITH_EDITOR
  	static bool Load(TSharedPtr<FImplBackEndUEOnly>& InOutImplBackEndUEOnly, const FString& InModelFullFilePath);
- #endif //WITH_EDITOR
 
 	//static bool Load(TSharedPtr<FImplBackEndUEOnly>& InOutImplBackEndUEOnly, FNeuralTensorManager& InTensorManager, const TArray<TSharedPtr<FNeuralOperator>>& InOperators);
 
