NeuralNetworkInference (NNI) plugin:
- FNeuralTensors disappears. Instead:
1. UNN has TArray<FNeuralTensor> instead of FNeuralTensors.
2. No-longer-needed FNeuralTensors removed.

#rnx

diff --git a/Engine/Plugins/Experimental/NeuralNetworkInference/Source/NeuralNetworkInference/Internal/NeuralNetworkLegacy.h b/Engine/Plugins/Experimental/NeuralNetworkInference/Source/NeuralNetworkInference/Internal/NeuralNetworkLegacy.h
--- a/Engine/Plugins/Experimental/NeuralNetworkInference/Source/NeuralNetworkInference/Internal/NeuralNetworkLegacy.h
+++ b/Engine/Plugins/Experimental/NeuralNetworkInference/Source/NeuralNetworkInference/Internal/NeuralNetworkLegacy.h
@@ -5,7 +5,6 @@
 #include "CoreMinimal.h"
 #include "NeuralEnumClasses.h"
 #include "NeuralTensor.h"
-#include "NeuralTensors.h"
 
 #include "NeuralNetwork.h"
 #include "ModelProto.h"
diff --git a/Engine/Plugins/Experimental/NeuralNetworkInference/Source/NeuralNetworkInference/Private/NeuralNetwork.cpp b/Engine/Plugins/Experimental/NeuralNetworkInference/Source/NeuralNetworkInference/Private/NeuralNetwork.cpp
--- a/Engine/Plugins/Experimental/NeuralNetworkInference/Source/NeuralNetworkInference/Private/NeuralNetwork.cpp
+++ b/Engine/Plugins/Experimental/NeuralNetworkInference/Source/NeuralNetworkInference/Private/NeuralNetwork.cpp
@@ -222,30 +222,89 @@ void UNeuralNetwork::SetBackEnd(const ENeuralBackEnd InBackEnd)
 
 const FNeuralTensor& UNeuralNetwork::GetInputTensor(const int32 InTensorIndex) const
 {
-	return InputTensors.GetTensor(InTensorIndex);
+	// UEAndORT
+	if (BackEndForCurrentPlatform == ENeuralBackEnd::UEAndORT)
+	{
+		return InputTensors[InTensorIndex];
+	}
+
+	// UEOnly
+	else if (BackEndForCurrentPlatform == ENeuralBackEnd::UEOnly)
+	{
+		const int32 TensorIndex = ImplBackEndUEOnly->TensorManager.GetInputIndexes()[InTensorIndex];
+		return ImplBackEndUEOnly->TensorManager.GetTensors()[TensorIndex];
+	}
+
+	// Unknown
+	checkf(false, TEXT("UNeuralNetwork::GetInputTensor(): Unknown [BackEnd,BackEndForCurrentPlatform] = [%d,%d]."), (int32)BackEnd, (int32)BackEndForCurrentPlatform);
+	return InputTensors[InTensorIndex];
 }
 
-const FNeuralTensors& UNeuralNetwork::GetInputTensors() const
+const TArray<FNeuralTensor>& UNeuralNetwork::GetInputTensors() const
 {
 	return InputTensors;
 }
 
 void UNeuralNetwork::SetInputFromArrayCopy(const TArray<float>& InArray, const int32 InTensorIndex)
 {
-	InputTensors.SetFromArrayCopy(InArray, InTensorIndex);
+	// UEAndORT
+	if (BackEndForCurrentPlatform == ENeuralBackEnd::UEAndORT)
+	{
+		InputTensors[InTensorIndex].SetFromArrayCopy(InArray);
+	}
+
+	// UEOnly
+	else if (BackEndForCurrentPlatform == ENeuralBackEnd::UEOnly)
+	{
+UE_LOG(LogNeuralNetworkInference, Warning, TEXT("UNeuralNetwork::SetInputFromArrayCopy(): UEOnly not implemented yet."));
+	}
+
+	// Unknown
+	else
+	{
+		UE_LOG(LogNeuralNetworkInference, Warning, TEXT("UNeuralNetwork::SetInputFromArrayCopy(): Unknown [BackEnd,BackEndForCurrentPlatform] = [%d,%d]."), (int32)BackEnd, (int32)BackEndForCurrentPlatform);
+	}
 }
 
 void* UNeuralNetwork::GetInputDataPointerMutable(const int32 InTensorIndex)
 {
-	return InputTensors.GetDataPointerMutable(InTensorIndex);
+	// UEAndORT
+	if (BackEndForCurrentPlatform == ENeuralBackEnd::UEAndORT)
+	{
+		return InputTensors[InTensorIndex].GetData(); // Or ImplBackEndUEAndORT->InputOrtTensors[InTensorIndex].GetTensorMutableData<float>();
+	}
+
+	// UEOnly
+	else if (BackEndForCurrentPlatform == ENeuralBackEnd::UEOnly)
+	{
+UE_LOG(LogNeuralNetworkInference, Warning, TEXT("UNeuralNetwork::GetInputDataPointerMutable(): UEOnly not implemented yet."));
+	}
+
+	// Unknown
+	UE_LOG(LogNeuralNetworkInference, Warning, TEXT("UNeuralNetwork::SetInputFromArrayCopy(): Unknown [BackEnd,BackEndForCurrentPlatform] = [%d,%d]."), (int32)BackEnd, (int32)BackEndForCurrentPlatform);
+	return nullptr;
 }
 
 const FNeuralTensor& UNeuralNetwork::GetOutputTensor(const int32 InTensorIndex) const
 {
-	return OutputTensors.GetTensor(InTensorIndex);
+	// UEAndORT
+	if (BackEndForCurrentPlatform == ENeuralBackEnd::UEAndORT)
+	{
+		return OutputTensors[InTensorIndex];
+	}
+
+	// UEOnly
+	else if (BackEndForCurrentPlatform == ENeuralBackEnd::UEOnly)
+	{
+UE_LOG(LogNeuralNetworkInference, Warning, TEXT("UNeuralNetwork::GetOutputTensor(): UEOnly not implemented yet."));
+	}
+
+	// Unknown
+	checkf(false, TEXT("UNeuralNetwork::GetOutputTensor(): Unknown [BackEnd,BackEndForCurrentPlatform] = [%d,%d]."), (int32)BackEnd, (int32)BackEndForCurrentPlatform);
+	return OutputTensors[InTensorIndex];
 }
 
-const FNeuralTensors& UNeuralNetwork::GetOutputTensors() const
+const TArray<FNeuralTensor>& UNeuralNetwork::GetOutputTensors() const
 {
 	return OutputTensors;
 }
@@ -264,7 +323,7 @@ void UNeuralNetwork::Run()
 	// UEAndORT
 	if (BackEndForCurrentPlatform == ENeuralBackEnd::UEAndORT)
 	{
-		ImplBackEndUEAndORT->Run(OutputTensors, InputTensors, SynchronousMode, InputDeviceType, OutputDeviceType);
+		ImplBackEndUEAndORT->Run(SynchronousMode, InputDeviceType, OutputDeviceType);
 	}
 
 	// UEOnly
diff --git a/Engine/Plugins/Experimental/NeuralNetworkInference/Source/NeuralNetworkInference/Private/NeuralNetworkImplBackEndUEAndORT.cpp b/Engine/Plugins/Experimental/NeuralNetworkInference/Source/NeuralNetworkInference/Private/NeuralNetworkImplBackEndUEAndORT.cpp
--- a/Engine/Plugins/Experimental/NeuralNetworkInference/Source/NeuralNetworkInference/Private/NeuralNetworkImplBackEndUEAndORT.cpp
+++ b/Engine/Plugins/Experimental/NeuralNetworkInference/Source/NeuralNetworkInference/Private/NeuralNetworkImplBackEndUEAndORT.cpp
@@ -9,7 +9,7 @@
  *****************************************************************************/
 
 //#if WITH_EDITOR
-//bool UNeuralNetwork::FImplBackEndUEAndORT::LoadFile(TSharedPtr<FImplBackEndUEAndORT>& InOutImplBackEndUEAndORT, FNeuralTensors& OutInputTensors, FNeuralTensors& OutOutputTensors,
+//bool UNeuralNetwork::FImplBackEndUEAndORT::LoadFile(TSharedPtr<FImplBackEndUEAndORT>& InOutImplBackEndUEAndORT, TArray<FNeuralTensor>& OutInputTensors, TArray<FNeuralTensor>& OutOutputTensors,
 //	TArray<bool>& OutAreInputTensorSizesVariable, const TArray<uint8>& InModelReadFromDiskInBytes, const FString& InModelFullFilePath, const ENeuralDeviceType InDeviceType)
 //{
 //	// Initialize and configure InOutImplBackEndUEAndORT
@@ -77,36 +77,48 @@
 //}
 //#endif //WITH_EDITOR
 
-bool UNeuralNetwork::FImplBackEndUEAndORT::Load(TSharedPtr<FImplBackEndUEAndORT>& InOutImplBackEndUEAndORT, FNeuralTensors& OutInputTensors, FNeuralTensors& OutOutputTensors,
+bool UNeuralNetwork::FImplBackEndUEAndORT::Load(TSharedPtr<FImplBackEndUEAndORT>& InOutImplBackEndUEAndORT, TArray<FNeuralTensor>& OutInputTensors, TArray<FNeuralTensor>& OutOutputTensors,
 	TArray<bool>& OutAreInputTensorSizesVariable, const TArray<uint8>& InModelReadFromDiskInBytes, const FString& InModelFullFilePath, const ENeuralDeviceType InDeviceType)
 {
 #ifdef WITH_FULL_NNI_SUPPORT
-	const FRedirectCoutAndCerrToUeLog RedirectCoutAndCerrToUeLog;
-
-	// Initialize and configure InOutImplBackEndUEAndORT
-	if (!UNeuralNetwork::FImplBackEndUEAndORT::InitializedAndConfigureMembers(InOutImplBackEndUEAndORT, InModelFullFilePath, InDeviceType))
+#if WITH_EDITOR
+	try
+#endif //WITH_EDITOR
 	{
-		UE_LOG(LogNeuralNetworkInference, Warning, TEXT("UNeuralNetwork::Load(): InitializedAndConfigureMembers failed."));
-		return false;
-	}
+		const FRedirectCoutAndCerrToUeLog RedirectCoutAndCerrToUeLog;
 
-	// Create session from model saved in InModelReadFromDiskInBytes (if not empty)
-	if (InModelReadFromDiskInBytes.Num() > 0)
-	{
-		// Read model from ModelReadFromDiskInBytesVector
-		InOutImplBackEndUEAndORT->Session = MakeUnique<Ort::Session>(*InOutImplBackEndUEAndORT->Environment, InModelReadFromDiskInBytes.GetData(), InModelReadFromDiskInBytes.Num(), *InOutImplBackEndUEAndORT->SessionOptions);
+		// Initialize and configure InOutImplBackEndUEAndORT
+		if (!UNeuralNetwork::FImplBackEndUEAndORT::InitializedAndConfigureMembers(InOutImplBackEndUEAndORT, InModelFullFilePath, InDeviceType))
+		{
+			UE_LOG(LogNeuralNetworkInference, Warning, TEXT("UNeuralNetwork::Load(): InitializedAndConfigureMembers failed."));
+			return false;
+		}
+
+		// Create session from model saved in InModelReadFromDiskInBytes (if not empty)
+		if (InModelReadFromDiskInBytes.Num() > 0)
+		{
+			// Read model from ModelReadFromDiskInBytesVector
+			InOutImplBackEndUEAndORT->Session = MakeUnique<Ort::Session>(*InOutImplBackEndUEAndORT->Environment, InModelReadFromDiskInBytes.GetData(), InModelReadFromDiskInBytes.Num(), *InOutImplBackEndUEAndORT->SessionOptions);
+		}
+		// Else
+		else
+		{
+			UE_LOG(LogNeuralNetworkInference, Warning, TEXT("InModelReadFromDiskInBytes was empty."));
+			return false;
+		}
+
+		InOutImplBackEndUEAndORT->ConfigureTensors(OutInputTensors, &OutAreInputTensorSizesVariable);
+		InOutImplBackEndUEAndORT->ConfigureTensors(OutOutputTensors);
+
+		return true;
 	}
-	// Else
-	else
+#if WITH_EDITOR
+	catch (const std::exception& Exception)
 	{
-		UE_LOG(LogNeuralNetworkInference, Warning, TEXT("InModelReadFromDiskInBytes was empty."));
+		UE_LOG(LogNeuralNetworkInference, Error, TEXT("%s"), UTF8_TO_TCHAR(Exception.what()));
 		return false;
 	}
-
-	InOutImplBackEndUEAndORT->ConfigureTensors(OutInputTensors, &OutAreInputTensorSizesVariable);
-	InOutImplBackEndUEAndORT->ConfigureTensors(OutOutputTensors);
-
-	return true;
+#endif //WITH_EDITOR
 
 #else //WITH_FULL_NNI_SUPPORT
 	UE_LOG(LogNeuralNetworkInference, Warning, TEXT("UNeuralNetwork::Load(): Platform or Operating System not suported yet for UEAndORT BackEnd. Set BackEnd to ENeuralBackEnd::Auto (recommended) or ENeuralBackEnd::UEOnly for this platform."));
@@ -114,32 +126,44 @@ bool UNeuralNetwork::FImplBackEndUEAndORT::Load(TSharedPtr<FImplBackEndUEAndORT>
 #endif //WITH_FULL_NNI_SUPPORT
 }
 
-void UNeuralNetwork::FImplBackEndUEAndORT::Run(FNeuralTensors& InOutOutputTensors, const FNeuralTensors& InInputTensors, const ENeuralNetworkSynchronousMode InSynchronousMode, const ENeuralDeviceType InInputDeviceType, const ENeuralDeviceType InOutputDeviceType)
+void UNeuralNetwork::FImplBackEndUEAndORT::Run(const ENeuralNetworkSynchronousMode InSynchronousMode, const ENeuralDeviceType InInputDeviceType, const ENeuralDeviceType InOutputDeviceType)
 {
 #ifdef WITH_FULL_NNI_SUPPORT
-	// @todo: Temporarily disabled until we connect GPU input/output between UE and ORT
-	if (InInputDeviceType == ENeuralDeviceType::GPU || InOutputDeviceType == ENeuralDeviceType::GPU)
-	{
-		UE_LOG(LogNeuralNetworkInference, Warning, TEXT("UNeuralNetwork::Run(): InputDeviceType & OutputDeviceType must be set to CPU for now."));
-		return;
-	}
-
-	// Run UNeuralNetwork
-	if (InSynchronousMode == ENeuralNetworkSynchronousMode::Synchronous)
+#if WITH_EDITOR
+	try
+#endif //WITH_EDITOR
 	{
 		const FRedirectCoutAndCerrToUeLog RedirectCoutAndCerrToUeLog;
-		Session->Run(Ort::RunOptions{ nullptr },
-			InInputTensors.GetTensorNames(), InInputTensors.GetONNXRuntimeTensors(), InInputTensors.GetNumberTensors(),
-			InOutOutputTensors.GetTensorNames(), InOutOutputTensors.GetONNXRuntimeTensors(), InOutOutputTensors.GetNumberTensors());
-	}
-	else if (InSynchronousMode == ENeuralNetworkSynchronousMode::Asynchronous)
-	{
-		UE_LOG(LogNeuralNetworkInference, Warning, TEXT("UNeuralNetwork::Run(): SynchronousMode = %d not implemented yet. Use SynchronousMode = Synchronous."), (int32)InSynchronousMode);
+
+		// @todo: Temporarily disabled until we connect GPU input/output between UE and ORT
+		if (InInputDeviceType == ENeuralDeviceType::GPU || InOutputDeviceType == ENeuralDeviceType::GPU)
+		{
+			UE_LOG(LogNeuralNetworkInference, Warning, TEXT("UNeuralNetwork::Run(): InputDeviceType & OutputDeviceType must be set to CPU for now."));
+			return;
+		}
+
+		// Run UNeuralNetwork
+		if (InSynchronousMode == ENeuralNetworkSynchronousMode::Synchronous)
+		{
+			Session->Run(Ort::RunOptions{ nullptr },
+				InputTensorNames.GetData(), &InputOrtTensors[0], InputTensorNames.Num(),
+				OutputTensorNames.GetData(), &OutputOrtTensors[0], OutputTensorNames.Num());
+		}
+		else if (InSynchronousMode == ENeuralNetworkSynchronousMode::Asynchronous)
+		{
+			UE_LOG(LogNeuralNetworkInference, Warning, TEXT("UNeuralNetwork::Run(): SynchronousMode = %d not implemented yet. Use SynchronousMode = Synchronous."), (int32)InSynchronousMode);
+		}
+		else
+		{
+			UE_LOG(LogNeuralNetworkInference, Warning, TEXT("UNeuralNetwork::Run(): Unknown SynchronousMode = %d."), (int32)InSynchronousMode);
+		}
 	}
-	else
+#if WITH_EDITOR
+	catch (const std::exception& Exception)
 	{
-		UE_LOG(LogNeuralNetworkInference, Warning, TEXT("UNeuralNetwork::Run(): Unknown SynchronousMode = %d."), (int32)InSynchronousMode);
+		UE_LOG(LogNeuralNetworkInference, Error, TEXT("%s"), UTF8_TO_TCHAR(Exception.what()));
 	}
+#endif //WITH_EDITOR
 
 #else //WITH_FULL_NNI_SUPPORT
 	UE_LOG(LogNeuralNetworkInference, Warning, TEXT("UNeuralNetwork::Run(): Platform or Operating System not suported yet for UEAndORT BackEnd. Set BackEnd to ENeuralBackEnd::Auto or ENeuralBackEnd::UEOnly for this platform."));
@@ -166,6 +190,8 @@ bool UNeuralNetwork::FImplBackEndUEAndORT::InitializedAndConfigureMembers(TShare
 		InOutImplBackEndUEAndORT->Environment = MakeUnique<Ort::Env>(ORT_LOGGING_LEVEL_WARNING, ModelFullFilePathCharPtr); // Any unique string would work, it does not need to be the file path
 
 		InOutImplBackEndUEAndORT->Allocator = MakeUnique<Ort::AllocatorWithDefaultOptions>();
+
+		InOutImplBackEndUEAndORT->AllocatorInfo = MakeUnique<Ort::MemoryInfo>(Ort::MemoryInfo::CreateCpu(OrtDeviceAllocator, OrtMemTypeCPU));
 	}
 
 	// Configure InOutImplBackEndUEAndORT
@@ -232,7 +258,7 @@ bool UNeuralNetwork::FImplBackEndUEAndORT::ConfigureMembers(const ENeuralDeviceT
 	return true;
 }
 
-void UNeuralNetwork::FImplBackEndUEAndORT::ConfigureTensors(FNeuralTensors& OutTensors, TArray<bool>* OutAreInputTensorSizesVariable)
+void UNeuralNetwork::FImplBackEndUEAndORT::ConfigureTensors(TArray<FNeuralTensor>& OutTensors, TArray<bool>* OutAreInputTensorSizesVariable)
 {
 	const bool bIsInput = (OutAreInputTensorSizesVariable != nullptr);
 	TArray<const char*> TensorNames;
@@ -322,7 +348,76 @@ void UNeuralNetwork::FImplBackEndUEAndORT::ConfigureTensors(FNeuralTensors& OutT
 
 	}
 
-	OutTensors.SetFromNetwork(TensorNames, TensorDataTypes, TensorSizes);
+	SetTensorsFromNetwork(OutTensors, TensorNames, TensorDataTypes, TensorSizes, bIsInput);
+}
+
+void UNeuralNetwork::FImplBackEndUEAndORT::SetTensorsFromNetwork(TArray<FNeuralTensor>& OutTensors, TArray<const char*>& InTensorNames, TArray<ENeuralDataType>& InTensorDataTypes,
+	TArray<TArray<int64>>& InSizes, const bool bIsInput)
+{
+	const int32 TensorNumber = InTensorNames.Num();
+	if (InTensorDataTypes.Num() != TensorNumber || InSizes.Num() != TensorNumber)
+	{
+		UE_LOG(LogNeuralNetworkInference, Warning, TEXT("InTensorNames.Num() == InTensorDataTypes.Num() == InSizes.Num() failed, %d vs. %d vs. %d."),
+			InTensorNames.Num(), InTensorDataTypes.Num(), InSizes.Num());
+		return;
+	}
+
+	// Swap variables
+	TArray<const char*>& TensorNames = (bIsInput ? InputTensorNames : OutputTensorNames);
+	Swap(TensorNames, InTensorNames);
+
+	// Pre-allocate each variable
+	if (OutTensors.Num() != TensorNumber)
+	{
+		OutTensors.SetNum(TensorNumber);
+	}
+
+	// Config each TensorIndex
+	TArray<Ort::Value>& OrtTensors = (bIsInput ? InputOrtTensors : OutputOrtTensors);
+	for (int32 TensorIndex = 0; TensorIndex < TensorNumber; ++TensorIndex)
+	{
+		if (OrtTensors.Num() <= TensorIndex)
+		{
+			OrtTensors.Emplace(Ort::Value(nullptr));
+		}
+		// Pre-allocate TArray
+		OutTensors[TensorIndex].SetNumUninitialized(InSizes[TensorIndex], InTensorDataTypes[TensorIndex]);
+		// Link tensor with ORT blob
+		LinkTensorToONNXRuntime(OutTensors, OrtTensors, *AllocatorInfo, TensorIndex);
+	}
+}
+
+void UNeuralNetwork::FImplBackEndUEAndORT::LinkTensorToONNXRuntime(TArray<FNeuralTensor>& InOutTensors, TArray<Ort::Value>& InOutOrtTensors, Ort::MemoryInfo& InOutAllocatorInfo, const int32 InTensorIndex)
+{
+	const TArray<int64>& Sizes = InOutTensors[InTensorIndex].GetSizes();
+	if (Sizes.Num() > 0 && InOutTensors[InTensorIndex].Num() > 0)
+	{
+		FNeuralTensor& Tensor = InOutTensors[InTensorIndex];
+		const int64 Volume = Tensor.Num();
+		const int32 ArrayDimensions = Sizes.Num();
+
+		const ENeuralDataType NeuralDataType = Tensor.GetDataType();
+		if (NeuralDataType == ENeuralDataType::Float)
+		{
+#ifdef _WIN32
+			const TArray<int64_t>& SizesInt64t = Sizes;
+#else
+			checkf(sizeof(int64) == sizeof(int64_t), TEXT("int64 and int64_t should both have the same size."));
+			TArray<int64_t> SizesInt64t;
+			SizesInt64t.SetNumUninitialized(ArrayDimensions);
+			FMemory::Memcpy(SizesInt64t.GetData(), (int64_t*)Sizes.GetData(), sizeof(int64_t) * ArrayDimensions);
+#endif //_WIN32
+			InOutOrtTensors[InTensorIndex] = Ort::Value::CreateTensor<float>(InOutAllocatorInfo, Tensor.GetDataCasted<float>(), Volume, SizesInt64t.GetData(), ArrayDimensions);
+		}
+		//else if (NeuralDataType == ENeuralDataType::Double)
+		//{
+		//	InOutOrtTensors[InTensorIndex] = Ort::Value::CreateTensor<double>(InOutAllocatorInfo, Tensor.GetDataCasted<double>(), Volume, Sizes.GetData(), ArrayDimensions);
+		//}
+		else
+		{
+			UE_LOG(LogNeuralNetworkInference, Warning, TEXT("Not implemented (yet) for ENeuralDataType = %d."), (int32)NeuralDataType);
+		}
+	}
 }
 
 #endif //WITH_FULL_NNI_SUPPORT
diff --git a/Engine/Plugins/Experimental/NeuralNetworkInference/Source/NeuralNetworkInference/Private/NeuralNetworkImplBackEndUEAndORT.h b/Engine/Plugins/Experimental/NeuralNetworkInference/Source/NeuralNetworkInference/Private/NeuralNetworkImplBackEndUEAndORT.h
--- a/Engine/Plugins/Experimental/NeuralNetworkInference/Source/NeuralNetworkInference/Private/NeuralNetworkImplBackEndUEAndORT.h
+++ b/Engine/Plugins/Experimental/NeuralNetworkInference/Source/NeuralNetworkInference/Private/NeuralNetworkImplBackEndUEAndORT.h
@@ -29,21 +29,28 @@ struct UNeuralNetwork::FImplBackEndUEAndORT
 {
 public:
 #ifdef WITH_FULL_NNI_SUPPORT
+	// Network-related variables
 	TUniquePtr<Ort::Env> Environment;
 	TUniquePtr<Ort::Session> Session;
 	TUniquePtr<Ort::AllocatorWithDefaultOptions> Allocator;
 	TUniquePtr<Ort::SessionOptions> SessionOptions;
+	// Tensor-related variables
+	TUniquePtr<Ort::MemoryInfo> AllocatorInfo; /* Memory allocator information */
+	TArray<Ort::Value> InputOrtTensors; /* Actual ONNXRuntime tensors */
+	TArray<const char*> InputTensorNames; /* Tensor names */
+	TArray<Ort::Value> OutputOrtTensors; /* Actual ONNXRuntime tensors */
+	TArray<const char*> OutputTensorNames; /* Tensor names */
 #endif //WITH_FULL_NNI_SUPPORT
 
 //#if WITH_EDITOR
-//	static bool LoadFile(TSharedPtr<FImplBackEndUEAndORT>& InOutImplBackEndUEAndORT, FNeuralTensors& OutInputTensors, FNeuralTensors& OutOutputTensors,
+//	static bool LoadFile(TSharedPtr<FImplBackEndUEAndORT>& InOutImplBackEndUEAndORT, TArray<FNeuralTensor>& OutInputTensors, TArray<FNeuralTensor>& OutOutputTensors,
 //		TArray<bool>& OutAreInputTensorSizesVariable, const TArray<uint8>& InModelReadFromDiskInBytes, const FString& InModelFullFilePath, const ENeuralDeviceType InDeviceType);
 //#endif //WITH_EDITOR
 
-	static bool Load(TSharedPtr<FImplBackEndUEAndORT>& InOutImplBackEndUEAndORT, FNeuralTensors& OutInputTensors, FNeuralTensors& OutOutputTensors,
+	static bool Load(TSharedPtr<FImplBackEndUEAndORT>& InOutImplBackEndUEAndORT, TArray<FNeuralTensor>& OutInputTensors, TArray<FNeuralTensor>& OutOutputTensors,
 		TArray<bool>& OutAreInputTensorSizesVariable, const TArray<uint8>& InModelReadFromDiskInBytes, const FString& InModelFullFilePath, const ENeuralDeviceType InDeviceType);
 	
-	void Run(FNeuralTensors& InOutOutputTensors, const FNeuralTensors& InInputTensors, const ENeuralNetworkSynchronousMode InSynchronousMode, const ENeuralDeviceType InInputDeviceType, const ENeuralDeviceType InOutputDeviceType);
+	void Run(const ENeuralNetworkSynchronousMode InSynchronousMode, const ENeuralDeviceType InInputDeviceType, const ENeuralDeviceType InOutputDeviceType);
 
 #ifdef WITH_FULL_NNI_SUPPORT
 private:
@@ -51,6 +58,10 @@ private:
 
 	bool ConfigureMembers(const ENeuralDeviceType InDeviceType);
 
-	void ConfigureTensors(FNeuralTensors& OutTensors, TArray<bool>* OutAreInputTensorSizesVariable = nullptr);
+	void ConfigureTensors(TArray<FNeuralTensor>& OutTensors, TArray<bool>* OutAreInputTensorSizesVariable = nullptr);
+
+	void SetTensorsFromNetwork(TArray<FNeuralTensor>& OutTensors, TArray<const char*>& InTensorNames, TArray<ENeuralDataType>& InTensorDataTypes, TArray<TArray<int64>>& InSizes, const bool bIsInput);
+
+	static void LinkTensorToONNXRuntime(TArray<FNeuralTensor>& InOutTensors, TArray<Ort::Value>& InOutOrtTensors, Ort::MemoryInfo& InOutAllocatorInfo, const int32 InTensorIndex);
 #endif //WITH_FULL_NNI_SUPPORT
 };
diff --git a/Engine/Plugins/Experimental/NeuralNetworkInference/Source/NeuralNetworkInference/Private/NeuralTensors.cpp b/Engine/Plugins/Experimental/NeuralNetworkInference/Source/NeuralNetworkInference/Private/NeuralTensors.cpp
deleted file mode 100644
diff --git a/Engine/Plugins/Experimental/NeuralNetworkInference/Source/NeuralNetworkInference/Public/NeuralNetwork.h b/Engine/Plugins/Experimental/NeuralNetworkInference/Source/NeuralNetworkInference/Public/NeuralNetwork.h
--- a/Engine/Plugins/Experimental/NeuralNetworkInference/Source/NeuralNetworkInference/Public/NeuralNetwork.h
+++ b/Engine/Plugins/Experimental/NeuralNetworkInference/Source/NeuralNetworkInference/Public/NeuralNetwork.h
@@ -5,7 +5,6 @@
 #include "CoreMinimal.h"
 #include "NeuralEnumClasses.h"
 #include "NeuralTensor.h"
-#include "NeuralTensors.h"
 #include "NeuralNetwork.generated.h"
 
 /**
@@ -139,7 +138,7 @@ public:
 	 * These functions either take a TArray as input, return a modificable void* to fill the data, or return a constant FNeuralTensor(s) to see input properties (e.g., size or dimensions).
 	 */
 	const FNeuralTensor& GetInputTensor(const int32 InTensorIndex = 0) const;
-	const FNeuralTensors& GetInputTensors() const;
+	const TArray<FNeuralTensor>& GetInputTensors() const;
 	void SetInputFromArrayCopy(const TArray<float>& InArray, const int32 InTensorIndex = 0);
 	void* GetInputDataPointerMutable(const int32 InTensorIndex = 0);
 
@@ -147,7 +146,7 @@ public:
 	 * Functions to get output. The returned FNeuralTensor(s) are constant to prevent the user from modifying the tensor properties (e.g., size or dimensions).
 	 */
 	const FNeuralTensor& GetOutputTensor(const int32 InTensorIndex = 0) const;
-	const FNeuralTensors& GetOutputTensors() const;
+	const TArray<FNeuralTensor>& GetOutputTensors() const;
 
 	/**
 	 * Run() executes the forward pass on the current UNeuralNetwork given the current input FDeprecatedNeuralTensor(s), which were previously filled with
@@ -197,13 +196,13 @@ protected:
 	FString ModelFullFilePath;
 
 	/**
-	 * InputTensors and OutputTensors represent the input and output TArray<FNeuralTensors> of the network, respectively.
+	 * InputTensors and OutputTensors represent the input and output TArray<FNeuralTensor> of the network, respectively.
 	 */
 	UPROPERTY(VisibleAnywhere, Category = "Neural Network Inference")
-	FNeuralTensors InputTensors;
+	TArray<FNeuralTensor> InputTensors;
 
 	UPROPERTY(VisibleAnywhere, Category = "Neural Network Inference")
-	FNeuralTensors OutputTensors;
+	TArray<FNeuralTensor> OutputTensors;
 
 	/** Whether some of the FNeuralTensor of InputTensor have flexible/variable dimensions. */
 	UPROPERTY(VisibleAnywhere, Category = "Neural Network Inference")
diff --git a/Engine/Plugins/Experimental/NeuralNetworkInference/Source/NeuralNetworkInference/Public/NeuralTensors.h b/Engine/Plugins/Experimental/NeuralNetworkInference/Source/NeuralNetworkInference/Public/NeuralTensors.h
deleted file mode 100644
diff --git a/Engine/Plugins/Experimental/NeuralNetworkInference/Source/NeuralNetworkInferenceQA/Private/ModelUnitTester.cpp b/Engine/Plugins/Experimental/NeuralNetworkInference/Source/NeuralNetworkInferenceQA/Private/ModelUnitTester.cpp
--- a/Engine/Plugins/Experimental/NeuralNetworkInference/Source/NeuralNetworkInferenceQA/Private/ModelUnitTester.cpp
+++ b/Engine/Plugins/Experimental/NeuralNetworkInference/Source/NeuralNetworkInferenceQA/Private/ModelUnitTester.cpp
@@ -56,14 +56,14 @@ void FModelUnitTester::ModelLoadAccuracyAndSpeedTests(const FString& InProjectCo
 		}
 		// Input debugging
 		UE_LOG(LogNeuralNetworkInferenceQA, Display, TEXT("--------------- %s - Input/Output"), *ModelName);
-		for (uint32 TensorIndex = 0; TensorIndex < Network->GetInputTensors().GetNumberTensors(); ++TensorIndex)
+		for (int32 TensorIndex = 0; TensorIndex < Network->GetInputTensors().Num(); ++TensorIndex)
 		{
-			UE_LOG(LogNeuralNetworkInferenceQA, Display, TEXT("InputTensor[%d] = %s."), TensorIndex, *Network->GetInputTensors().GetTensorName(TensorIndex));
+			UE_LOG(LogNeuralNetworkInferenceQA, Display, TEXT("InputTensor[%d] = %s."), TensorIndex, *Network->GetInputTensor(TensorIndex).GetName());
 		}
 		// Output debugging
-		for (uint32 TensorIndex = 0; TensorIndex < Network->GetOutputTensors().GetNumberTensors(); ++TensorIndex)
+		for (int32 TensorIndex = 0; TensorIndex < Network->GetOutputTensors().Num(); ++TensorIndex)
 		{
-			UE_LOG(LogNeuralNetworkInferenceQA, Display, TEXT("OutputTensor[%d] = %s."), TensorIndex, *Network->GetOutputTensors().GetTensorName(TensorIndex));
+			UE_LOG(LogNeuralNetworkInferenceQA, Display, TEXT("OutputTensor[%d] = %s."), TensorIndex, *Network->GetOutputTensor(TensorIndex).GetName());
 		}
 		ModelAccuracyTest(Network, InInputArrayValues, CPUGroundTruths, GPUGroundTruths);
 
