NeuralNetworkInference (NNI) plugin:
- WITH_FULL_NNI_SUPPORT renamed as WITH_UE_AND_ORT_SUPPORT for additional clarity and to avoid future confusion.

#rnx

diff --git a/Engine/Plugins/Experimental/NeuralNetworkInference/Source/NeuralNetworkInference/NeuralNetworkInference.Build.cs b/Engine/Plugins/Experimental/NeuralNetworkInference/Source/NeuralNetworkInference/NeuralNetworkInference.Build.cs
--- a/Engine/Plugins/Experimental/NeuralNetworkInference/Source/NeuralNetworkInference/NeuralNetworkInference.Build.cs
+++ b/Engine/Plugins/Experimental/NeuralNetworkInference/Source/NeuralNetworkInference/NeuralNetworkInference.Build.cs
@@ -10,7 +10,7 @@ public class NeuralNetworkInference : ModuleRules
 		bool bIsORTSupported = (Target.Platform == UnrealTargetPlatform.Win64 || Target.Platform == UnrealTargetPlatform.Linux);
 		if (bIsORTSupported)
 		{
-			PublicDefinitions.Add("WITH_FULL_NNI_SUPPORT");
+			PublicDefinitions.Add("WITH_UE_AND_ORT_SUPPORT");
 		}
 
 		ShortName = "NNI"; // Could be removed when plugin moves to Experimental, NFL path is too long
diff --git a/Engine/Plugins/Experimental/NeuralNetworkInference/Source/NeuralNetworkInference/Private/NeuralNetwork.cpp b/Engine/Plugins/Experimental/NeuralNetworkInference/Source/NeuralNetworkInference/Private/NeuralNetwork.cpp
--- a/Engine/Plugins/Experimental/NeuralNetworkInference/Source/NeuralNetworkInference/Private/NeuralNetwork.cpp
+++ b/Engine/Plugins/Experimental/NeuralNetworkInference/Source/NeuralNetworkInference/Private/NeuralNetwork.cpp
@@ -25,11 +25,11 @@ ENeuralBackEnd FPrivateNeuralNetwork::SetBackEndForCurrentPlatform(const ENeural
 	// Auto
 	if (InBackEnd == ENeuralBackEnd::Auto)
 	{
-#ifdef WITH_FULL_NNI_SUPPORT
+#ifdef WITH_UE_AND_ORT_SUPPORT
 		return ENeuralBackEnd::UEAndORT;
-#else //WITH_FULL_NNI_SUPPORT
+#else //WITH_UE_AND_ORT_SUPPORT
 		return ENeuralBackEnd::UEOnly;
-#endif //WITH_FULL_NNI_SUPPORT
+#endif //WITH_UE_AND_ORT_SUPPORT
 	}
 	// Otherwise
 	return InBackEnd;
diff --git a/Engine/Plugins/Experimental/NeuralNetworkInference/Source/NeuralNetworkInference/Private/NeuralNetworkImplBackEndUEAndORT.cpp b/Engine/Plugins/Experimental/NeuralNetworkInference/Source/NeuralNetworkInference/Private/NeuralNetworkImplBackEndUEAndORT.cpp
--- a/Engine/Plugins/Experimental/NeuralNetworkInference/Source/NeuralNetworkInference/Private/NeuralNetworkImplBackEndUEAndORT.cpp
+++ b/Engine/Plugins/Experimental/NeuralNetworkInference/Source/NeuralNetworkInference/Private/NeuralNetworkImplBackEndUEAndORT.cpp
@@ -80,7 +80,7 @@
 bool UNeuralNetwork::FImplBackEndUEAndORT::Load(TSharedPtr<FImplBackEndUEAndORT>& InOutImplBackEndUEAndORT, TArray<FNeuralTensor>& OutInputTensors, TArray<FNeuralTensor>& OutOutputTensors,
 	TArray<bool>& OutAreInputTensorSizesVariable, const TArray<uint8>& InModelReadFromDiskInBytes, const FString& InModelFullFilePath, const ENeuralDeviceType InDeviceType)
 {
-#ifdef WITH_FULL_NNI_SUPPORT
+#ifdef WITH_UE_AND_ORT_SUPPORT
 #if WITH_EDITOR
 	try
 #endif //WITH_EDITOR
@@ -120,15 +120,15 @@ bool UNeuralNetwork::FImplBackEndUEAndORT::Load(TSharedPtr<FImplBackEndUEAndORT>
 	}
 #endif //WITH_EDITOR
 
-#else //WITH_FULL_NNI_SUPPORT
+#else //WITH_UE_AND_ORT_SUPPORT
 	UE_LOG(LogNeuralNetworkInference, Warning, TEXT("UNeuralNetwork::Load(): Platform or Operating System not suported yet for UEAndORT BackEnd. Set BackEnd to ENeuralBackEnd::Auto (recommended) or ENeuralBackEnd::UEOnly for this platform."));
 	return false;
-#endif //WITH_FULL_NNI_SUPPORT
+#endif //WITH_UE_AND_ORT_SUPPORT
 }
 
 void UNeuralNetwork::FImplBackEndUEAndORT::Run(const ENeuralNetworkSynchronousMode InSynchronousMode, const ENeuralDeviceType InInputDeviceType, const ENeuralDeviceType InOutputDeviceType)
 {
-#ifdef WITH_FULL_NNI_SUPPORT
+#ifdef WITH_UE_AND_ORT_SUPPORT
 #if WITH_EDITOR
 	try
 #endif //WITH_EDITOR
@@ -165,9 +165,9 @@ void UNeuralNetwork::FImplBackEndUEAndORT::Run(const ENeuralNetworkSynchronousMo
 	}
 #endif //WITH_EDITOR
 
-#else //WITH_FULL_NNI_SUPPORT
+#else //WITH_UE_AND_ORT_SUPPORT
 	UE_LOG(LogNeuralNetworkInference, Warning, TEXT("UNeuralNetwork::Run(): Platform or Operating System not suported yet for UEAndORT BackEnd. Set BackEnd to ENeuralBackEnd::Auto or ENeuralBackEnd::UEOnly for this platform."));
-#endif //WITH_FULL_NNI_SUPPORT
+#endif //WITH_UE_AND_ORT_SUPPORT
 }
 
 
@@ -175,7 +175,7 @@ void UNeuralNetwork::FImplBackEndUEAndORT::Run(const ENeuralNetworkSynchronousMo
 /* UNeuralNetwork private functions
  *****************************************************************************/
 
-#ifdef WITH_FULL_NNI_SUPPORT
+#ifdef WITH_UE_AND_ORT_SUPPORT
 
 bool UNeuralNetwork::FImplBackEndUEAndORT::InitializedAndConfigureMembers(TSharedPtr<FImplBackEndUEAndORT>& InOutImplBackEndUEAndORT, const FString& InModelFullFilePath, const ENeuralDeviceType InDeviceType)
 {
@@ -420,4 +420,4 @@ void UNeuralNetwork::FImplBackEndUEAndORT::LinkTensorToONNXRuntime(TArray<FNeura
 	}
 }
 
-#endif //WITH_FULL_NNI_SUPPORT
+#endif //WITH_UE_AND_ORT_SUPPORT
diff --git a/Engine/Plugins/Experimental/NeuralNetworkInference/Source/NeuralNetworkInference/Private/NeuralNetworkImplBackEndUEAndORT.h b/Engine/Plugins/Experimental/NeuralNetworkInference/Source/NeuralNetworkInference/Private/NeuralNetworkImplBackEndUEAndORT.h
--- a/Engine/Plugins/Experimental/NeuralNetworkInference/Source/NeuralNetworkInference/Private/NeuralNetworkImplBackEndUEAndORT.h
+++ b/Engine/Plugins/Experimental/NeuralNetworkInference/Source/NeuralNetworkInference/Private/NeuralNetworkImplBackEndUEAndORT.h
@@ -10,7 +10,7 @@
 NNI_THIRD_PARTY_INCLUDES_START
 #undef check
 #undef TEXT
-#ifdef WITH_FULL_NNI_SUPPORT
+#ifdef WITH_UE_AND_ORT_SUPPORT
 	#include "RedirectCoutAndCerrToUeLog.h"
 
 	#include "onnxruntime/core/session/onnxruntime_cxx_api.h"
@@ -20,7 +20,7 @@ NNI_THIRD_PARTY_INCLUDES_START
 	#ifdef WITH_NNI_CPU_NOT_RECOMMENDED
 	#include "onnxruntime/core/providers/nni_cpu/nni_cpu_provider_factory.h"
 	#endif //WITH_NNI_CPU_NOT_RECOMMENDED
-#endif //WITH_FULL_NNI_SUPPORT
+#endif //WITH_UE_AND_ORT_SUPPORT
 NNI_THIRD_PARTY_INCLUDES_END
 
 
@@ -28,7 +28,7 @@ NNI_THIRD_PARTY_INCLUDES_END
 struct UNeuralNetwork::FImplBackEndUEAndORT
 {
 public:
-#ifdef WITH_FULL_NNI_SUPPORT
+#ifdef WITH_UE_AND_ORT_SUPPORT
 	// Network-related variables
 	TUniquePtr<Ort::Env> Environment;
 	TUniquePtr<Ort::Session> Session;
@@ -40,7 +40,7 @@ public:
 	TArray<const char*> InputTensorNames; /* Tensor names */
 	TArray<Ort::Value> OutputOrtTensors; /* Actual ONNXRuntime tensors */
 	TArray<const char*> OutputTensorNames; /* Tensor names */
-#endif //WITH_FULL_NNI_SUPPORT
+#endif //WITH_UE_AND_ORT_SUPPORT
 
 //#if WITH_EDITOR
 //	static bool LoadFile(TSharedPtr<FImplBackEndUEAndORT>& InOutImplBackEndUEAndORT, TArray<FNeuralTensor>& OutInputTensors, TArray<FNeuralTensor>& OutOutputTensors,
@@ -52,7 +52,7 @@ public:
 	
 	void Run(const ENeuralNetworkSynchronousMode InSynchronousMode, const ENeuralDeviceType InInputDeviceType, const ENeuralDeviceType InOutputDeviceType);
 
-#ifdef WITH_FULL_NNI_SUPPORT
+#ifdef WITH_UE_AND_ORT_SUPPORT
 private:
 	static bool InitializedAndConfigureMembers(TSharedPtr<FImplBackEndUEAndORT>& InOutImplBackEndUEAndORT, const FString& InModelFullFilePath, const ENeuralDeviceType InDeviceType);
 
@@ -63,5 +63,5 @@ private:
 	void SetTensorsFromNetwork(TArray<FNeuralTensor>& OutTensors, TArray<const char*>& InTensorNames, TArray<ENeuralDataType>& InTensorDataTypes, TArray<TArray<int64>>& InSizes, const bool bIsInput);
 
 	static void LinkTensorToONNXRuntime(TArray<FNeuralTensor>& InOutTensors, TArray<Ort::Value>& InOutOrtTensors, Ort::MemoryInfo& InOutAllocatorInfo, const int32 InTensorIndex);
-#endif //WITH_FULL_NNI_SUPPORT
+#endif //WITH_UE_AND_ORT_SUPPORT
 };
diff --git a/Engine/Plugins/Experimental/NeuralNetworkInference/Source/NeuralNetworkInferenceCore/Public/NeuralEnumClasses.h b/Engine/Plugins/Experimental/NeuralNetworkInference/Source/NeuralNetworkInferenceCore/Public/NeuralEnumClasses.h
--- a/Engine/Plugins/Experimental/NeuralNetworkInference/Source/NeuralNetworkInferenceCore/Public/NeuralEnumClasses.h
+++ b/Engine/Plugins/Experimental/NeuralNetworkInference/Source/NeuralNetworkInferenceCore/Public/NeuralEnumClasses.h
@@ -14,7 +14,7 @@
 UENUM()
 enum class ENeuralBackEnd : uint8
 {
-	UEAndORT, /* UnrealEngine-and-ONNXRuntime-accelerated back end, ideal for those platforms that support it. */
+	UEAndORT, /* UnrealEngine-and-ONNXRuntime-accelerated back end, ideal for those platforms that support it. WITH_UE_AND_ORT_SUPPORT is the C++ define macro that checks whether UEAndORT support exists for the current platform.*/
 	UEOnly, /* It might be slower than the UEAndORT back end, but it will compile in all platforms and OSs compatible with Unreal Engine. */
 	Auto /* Recommended setting. It will use UEAndORT if possible, and fall back to UEOnly only for those platforms that do not support UEAndORT yet. */
 };
diff --git a/Engine/Plugins/Experimental/NeuralNetworkInference/Source/NeuralNetworkInferenceEditor/Private/NeuralNetworkInferenceEditorModule.cpp b/Engine/Plugins/Experimental/NeuralNetworkInference/Source/NeuralNetworkInferenceEditor/Private/NeuralNetworkInferenceEditorModule.cpp
--- a/Engine/Plugins/Experimental/NeuralNetworkInference/Source/NeuralNetworkInferenceEditor/Private/NeuralNetworkInferenceEditorModule.cpp
+++ b/Engine/Plugins/Experimental/NeuralNetworkInference/Source/NeuralNetworkInferenceEditor/Private/NeuralNetworkInferenceEditorModule.cpp
@@ -36,7 +36,7 @@ void FNeuralNetworkInferenceEditorModule::StartupModule()
 {
 	// UNeuralNetwork - Register asset types
 	IAssetTools& AssetTools = FModuleManager::LoadModuleChecked<FAssetToolsModule>("AssetTools").Get();
-#ifdef WITH_FULL_NNI_SUPPORT
+#ifdef WITH_UE_AND_ORT_SUPPORT
 	Actions.Emplace(MakeShared<FNeuralNetworkAssetTypeActions>());
 	AssetTools.RegisterAssetTypeActions(Actions.Last().ToSharedRef());
 #endif
diff --git a/Engine/Plugins/Experimental/NeuralNetworkInference/Source/NeuralNetworkInferenceQA/Private/ModelUnitTester.cpp b/Engine/Plugins/Experimental/NeuralNetworkInference/Source/NeuralNetworkInferenceQA/Private/ModelUnitTester.cpp
--- a/Engine/Plugins/Experimental/NeuralNetworkInference/Source/NeuralNetworkInferenceQA/Private/ModelUnitTester.cpp
+++ b/Engine/Plugins/Experimental/NeuralNetworkInference/Source/NeuralNetworkInferenceQA/Private/ModelUnitTester.cpp
@@ -12,7 +12,7 @@
 
 void FModelUnitTester::GlobalTest(const FString& InProjectContentDir, const FString& InModelZooRelativeDirectory)
 {
-#ifdef WITH_FULL_NNI_SUPPORT
+#ifdef WITH_UE_AND_ORT_SUPPORT
 	// Model load, accuracy, and speed test
 	const TArray<FString> ModelNames({ TEXT("MLRigDeformer"), TEXT("cloth_network"), TEXT("HS"), TEXT("RL") });
 	const TArray<float> InputArrayValues({ 1.f, 0.f, -1.f, 100.f, -100.f, 0.5f, -0.5f }); // This one can be shorter than CPU/GPUGroundTruths
@@ -24,8 +24,8 @@ void FModelUnitTester::GlobalTest(const FString& InProjectContentDir, const FStr
 	const TArray<int32> GPURepetitions({ 1000, 1000, 100, 1000 });
 	ModelLoadAccuracyAndSpeedTests(InProjectContentDir, InModelZooRelativeDirectory, ModelNames, InputArrayValues, CPUGroundTruths, GPUGroundTruths, CPURepetitions, GPURepetitions);
 #else
-	UE_LOG(LogNeuralNetworkInferenceQA, Display, TEXT("--------------- FModelUnitTester test skipped (only if WITH_FULL_NNI_SUPPORT)."));
-#endif //WITH_FULL_NNI_SUPPORT
+	UE_LOG(LogNeuralNetworkInferenceQA, Display, TEXT("--------------- FModelUnitTester test skipped (only if WITH_UE_AND_ORT_SUPPORT)."));
+#endif //WITH_UE_AND_ORT_SUPPORT
 }
 
 
