NeuralNetworkInference (NNI) plugin:
- ModelReadFromDiskInBytes renamed as ModelReadFromFileInBytes.
-UNN: Load(FString) working in game for UEAndORT back end.
-UNN: Added Load(TArray<bytes>) to give even more flexibility to the code.

#rnx

diff --git a/Engine/Plugins/Experimental/NeuralNetworkInference/Source/NeuralNetworkInference/Private/NeuralNetwork.cpp b/Engine/Plugins/Experimental/NeuralNetworkInference/Source/NeuralNetworkInference/Private/NeuralNetwork.cpp
--- a/Engine/Plugins/Experimental/NeuralNetworkInference/Source/NeuralNetworkInference/Private/NeuralNetwork.cpp
+++ b/Engine/Plugins/Experimental/NeuralNetworkInference/Source/NeuralNetworkInference/Private/NeuralNetwork.cpp
@@ -60,15 +60,14 @@ UNeuralNetwork::~UNeuralNetwork()
 /* UNeuralNetwork public functions
  *****************************************************************************/
 
-#if WITH_EDITOR
 bool UNeuralNetwork::Load(const FString& InModelFilePath)
 {
-	DECLARE_SCOPE_CYCLE_COUNTER(TEXT("UNeuralNetwork_Load_EditorOnly"), STAT_UNeuralNetwork_Load, STATGROUP_MachineLearning);
+	DECLARE_SCOPE_CYCLE_COUNTER(TEXT("UNeuralNetwork_Load_FromFString"), STAT_UNeuralNetwork_Load, STATGROUP_MachineLearning);
 
 	// Clean previous networks
 	bIsLoaded = false;
 
-	// Fill ModelFullFilePath & ModelReadFromDiskInBytes
+	// Fill ModelFullFilePath & ModelReadFromFileInBytes
 	{
 		// Sanity check
 		if (InModelFilePath.IsEmpty())
@@ -84,34 +83,37 @@ bool UNeuralNetwork::Load(const FString& InModelFilePath)
 			UE_LOG(LogNeuralNetworkInference, Warning, TEXT("UNeuralNetwork::Load(): Model not found \"%s\"."), *ModelFullFilePath);
 			return false;
 		}
-		// Read file into ModelReadFromDiskInBytes
+		// Read file into ModelReadFromFileInBytes
 		// Source: [REDACTED_PROFILE_URL]/onnxruntime/blob/894fc828587c919d815918c4da6cde314e5d54ed/onnxruntime/test/shared_lib/test_model_loading.cc#L21-L31
-		if (!FFileHelper::LoadFileToArray(ModelReadFromDiskInBytes, *ModelFullFilePath))
+		if (!FFileHelper::LoadFileToArray(ModelReadFromFileInBytes, *ModelFullFilePath))
 		{
 			UE_LOG(LogNeuralNetworkInference, Warning, TEXT("UNeuralNetwork::Load(): Error reading model \"%s\"."), *ModelFullFilePath);
 			return false;
 		}
 	}
 
-	// UEAndORT
-	if (BackEndForCurrentPlatform == ENeuralBackEnd::UEAndORT)
-	{
-		return Load();
-	}
-	// UEOnly
-	else if (BackEndForCurrentPlatform == ENeuralBackEnd::UEOnly)
-	{
-		bIsLoaded = UNeuralNetwork::FImplBackEndUEOnly::Load(ImplBackEndUEOnly, /*ModelReadFromDiskInBytes*/ ModelFullFilePath);
-	}
-	// Unknown
-	else
+	return Load();
+}
+
+
+bool UNeuralNetwork::Load(TArray<uint8>& InModelReadFromFileInBytes)
+{
+	DECLARE_SCOPE_CYCLE_COUNTER(TEXT("UNeuralNetwork_Load_FromTArrayUInt8"), STAT_UNeuralNetwork_Load, STATGROUP_MachineLearning);
+
+	// Clean previous networks
+	bIsLoaded = false;
+
+	if (InModelReadFromFileInBytes.IsEmpty())
 	{
-		UE_LOG(LogNeuralNetworkInference, Warning, TEXT("UNeuralNetwork::Load(): Unknown [BackEnd,BackEndForCurrentPlatform] = [%d,%d]."), (int32)BackEnd, (int32)BackEndForCurrentPlatform);
+		UE_LOG(LogNeuralNetworkInference, Warning, TEXT("UNeuralNetwork::Load(): InModelReadFromFileInBytes is empty."));
+		return false;
 	}
 
-	return bIsLoaded;
+	// Read file into ModelReadFromFileInBytes
+	Swap(ModelReadFromFileInBytes, InModelReadFromFileInBytes);
+
+	return Load();
 }
-#endif //WITH_EDITOR
 
 
 bool UNeuralNetwork::Load()
@@ -124,12 +126,16 @@ bool UNeuralNetwork::Load()
 	// UEAndORT
 	if (BackEndForCurrentPlatform == ENeuralBackEnd::UEAndORT)
 	{
-		bIsLoaded = UNeuralNetwork::FImplBackEndUEAndORT::Load(ImplBackEndUEAndORT, InputTensors, OutputTensors, AreInputTensorSizesVariable, ModelReadFromDiskInBytes, ModelFullFilePath, GetDeviceType());
+		bIsLoaded = UNeuralNetwork::FImplBackEndUEAndORT::Load(ImplBackEndUEAndORT, InputTensors, OutputTensors, AreInputTensorSizesVariable, ModelReadFromFileInBytes, ModelFullFilePath, GetDeviceType());
 	}
 	// UEOnly
 	else if (BackEndForCurrentPlatform == ENeuralBackEnd::UEOnly)
 	{
-		bIsLoaded = UNeuralNetwork::FImplBackEndUEOnly::Load(ImplBackEndUEOnly, /*ModelReadFromDiskInBytes*/ ModelFullFilePath);
+#if WITH_EDITOR
+		bIsLoaded = UNeuralNetwork::FImplBackEndUEOnly::Load(ImplBackEndUEOnly, /*ModelReadFromFileInBytes*/ ModelFullFilePath);
+#else //WITH_EDITOR
+		UE_LOG(LogNeuralNetworkInference, Warning, TEXT("UNeuralNetwork::Load(const FString& InModelFilePath) only supported in Editor (WITH_EDITOR) for ENeuralBackEnd::UEOnly."));
+#endif //WITH_EDITOR
 	}
 	// Unknown
 	else
@@ -390,8 +396,8 @@ void UNeuralNetwork::PostInitProperties()
 void UNeuralNetwork::PostLoad()
 {
 	Super::PostLoad();
-	// If ModelReadFromDiskInBytes is not empty, call Load() 
-	if (ModelReadFromDiskInBytes.Num() > 0)
+	// If ModelReadFromFileInBytes is not empty, call Load() 
+	if (ModelReadFromFileInBytes.Num() > 0)
 	{
 		if (!Load())
 		{
diff --git a/Engine/Plugins/Experimental/NeuralNetworkInference/Source/NeuralNetworkInference/Private/NeuralNetworkImplBackEndUEAndORT.cpp b/Engine/Plugins/Experimental/NeuralNetworkInference/Source/NeuralNetworkInference/Private/NeuralNetworkImplBackEndUEAndORT.cpp
--- a/Engine/Plugins/Experimental/NeuralNetworkInference/Source/NeuralNetworkInference/Private/NeuralNetworkImplBackEndUEAndORT.cpp
+++ b/Engine/Plugins/Experimental/NeuralNetworkInference/Source/NeuralNetworkInference/Private/NeuralNetworkImplBackEndUEAndORT.cpp
@@ -10,7 +10,7 @@
 
 //#if WITH_EDITOR
 //bool UNeuralNetwork::FImplBackEndUEAndORT::LoadFile(TSharedPtr<FImplBackEndUEAndORT>& InOutImplBackEndUEAndORT, TArray<FNeuralTensor>& OutInputTensors, TArray<FNeuralTensor>& OutOutputTensors,
-//	TArray<bool>& OutAreInputTensorSizesVariable, const TArray<uint8>& InModelReadFromDiskInBytes, const FString& InModelFullFilePath, const ENeuralDeviceType InDeviceType)
+//	TArray<bool>& OutAreInputTensorSizesVariable, const TArray<uint8>& InModelReadFromFileInBytes, const FString& InModelFullFilePath, const ENeuralDeviceType InDeviceType)
 //{
 //	// Initialize and configure InOutImplBackEndUEAndORT
 //	const FRedirectCoutAndCerrToUeLog RedirectCoutAndCerrToUeLog;
@@ -21,7 +21,7 @@
 //	}
 //
 //	// Note: This code uses the changes in the ONNX Runtime API, which are not needed for desktop platforms (where ONNX/ProtoBuf is supported)
-//	// Fill InModelReadFromDiskInBytes from ONNX/ORT file
+//	// Fill InModelReadFromFileInBytes from ONNX/ORT file
 //	const FString FileExtension = FPaths::GetExtension(InModelFullFilePath, /*bIncludeDot*/ false);
 //	const char* const FilePathCharPtr = TCHAR_TO_ANSI(*InModelFullFilePath);
 //	// If ONNX file, turn into ORT format first
@@ -48,25 +48,25 @@
 //		// Read model from OutputORTOptimizedModelPath
 //		return Load(OutputORTOptimizedModelPath);
 //	}
-//	// Create session (it should be ORT file at this point), and read InModelReadFromDiskInBytes if not empty
+//	// Create session (it should be ORT file at this point), and read InModelReadFromFileInBytes if not empty
 //	else if (FileExtension.Equals(TEXT("ort"), ESearchCase::IgnoreCase))
 //	{
 //		// Read model from InModelFullFilePath
-//		std::vector<uint8_t> OutputModelReadFromDiskInBytesVector;
+//		std::vector<uint8_t> OutputModelReadFromFileInBytesVector;
 //		InOutImplBackEndUEAndORT->Session = MakeUnique<Ort::Session>(*InOutImplBackEndUEAndORT->Environment,
 //#ifdef _WIN32
 //			*InModelFullFilePath,
 //#else
 //			FilePathCharPtr,
 //#endif
-//			*InOutImplBackEndUEAndORT->SessionOptions, &OutputModelReadFromDiskInBytesVector);
+//			*InOutImplBackEndUEAndORT->SessionOptions, &OutputModelReadFromFileInBytesVector);
 //	
-//		// Fill InModelReadFromDiskInBytes
-//		const int32 ArraySize = OutputModelReadFromDiskInBytesVector.size();
+//		// Fill InModelReadFromFileInBytes
+//		const int32 ArraySize = OutputModelReadFromFileInBytesVector.size();
 //		if (ArraySize > 0)
 //		{
-//			InModelReadFromDiskInBytes.SetNumUninitialized(ArraySize);
-//			FMemory::Memcpy(InModelReadFromDiskInBytes.GetData(), &OutputModelReadFromDiskInBytesVector[0], ArraySize);
+//			InModelReadFromFileInBytes.SetNumUninitialized(ArraySize);
+//			FMemory::Memcpy(InModelReadFromFileInBytes.GetData(), &OutputModelReadFromFileInBytesVector[0], ArraySize);
 //		}
 //	
 //		return Load();
@@ -78,7 +78,7 @@
 //#endif //WITH_EDITOR
 
 bool UNeuralNetwork::FImplBackEndUEAndORT::Load(TSharedPtr<FImplBackEndUEAndORT>& InOutImplBackEndUEAndORT, TArray<FNeuralTensor>& OutInputTensors, TArray<FNeuralTensor>& OutOutputTensors,
-	TArray<bool>& OutAreInputTensorSizesVariable, const TArray<uint8>& InModelReadFromDiskInBytes, const FString& InModelFullFilePath, const ENeuralDeviceType InDeviceType)
+	TArray<bool>& OutAreInputTensorSizesVariable, const TArray<uint8>& InModelReadFromFileInBytes, const FString& InModelFullFilePath, const ENeuralDeviceType InDeviceType)
 {
 #ifdef WITH_UE_AND_ORT_SUPPORT
 #if WITH_EDITOR
@@ -94,16 +94,16 @@ bool UNeuralNetwork::FImplBackEndUEAndORT::Load(TSharedPtr<FImplBackEndUEAndORT>
 			return false;
 		}
 
-		// Create session from model saved in InModelReadFromDiskInBytes (if not empty)
-		if (InModelReadFromDiskInBytes.Num() > 0)
+		// Create session from model saved in InModelReadFromFileInBytes (if not empty)
+		if (InModelReadFromFileInBytes.Num() > 0)
 		{
-			// Read model from ModelReadFromDiskInBytesVector
-			InOutImplBackEndUEAndORT->Session = MakeUnique<Ort::Session>(*InOutImplBackEndUEAndORT->Environment, InModelReadFromDiskInBytes.GetData(), InModelReadFromDiskInBytes.Num(), *InOutImplBackEndUEAndORT->SessionOptions);
+			// Read model from ModelReadFromFileInBytesVector
+			InOutImplBackEndUEAndORT->Session = MakeUnique<Ort::Session>(*InOutImplBackEndUEAndORT->Environment, InModelReadFromFileInBytes.GetData(), InModelReadFromFileInBytes.Num(), *InOutImplBackEndUEAndORT->SessionOptions);
 		}
 		// Else
 		else
 		{
-			UE_LOG(LogNeuralNetworkInference, Warning, TEXT("InModelReadFromDiskInBytes was empty."));
+			UE_LOG(LogNeuralNetworkInference, Warning, TEXT("InModelReadFromFileInBytes was empty."));
 			return false;
 		}
 
diff --git a/Engine/Plugins/Experimental/NeuralNetworkInference/Source/NeuralNetworkInference/Private/NeuralNetworkImplBackEndUEAndORT.h b/Engine/Plugins/Experimental/NeuralNetworkInference/Source/NeuralNetworkInference/Private/NeuralNetworkImplBackEndUEAndORT.h
--- a/Engine/Plugins/Experimental/NeuralNetworkInference/Source/NeuralNetworkInference/Private/NeuralNetworkImplBackEndUEAndORT.h
+++ b/Engine/Plugins/Experimental/NeuralNetworkInference/Source/NeuralNetworkInference/Private/NeuralNetworkImplBackEndUEAndORT.h
@@ -44,11 +44,11 @@ public:
 
 //#if WITH_EDITOR
 //	static bool LoadFile(TSharedPtr<FImplBackEndUEAndORT>& InOutImplBackEndUEAndORT, TArray<FNeuralTensor>& OutInputTensors, TArray<FNeuralTensor>& OutOutputTensors,
-//		TArray<bool>& OutAreInputTensorSizesVariable, const TArray<uint8>& InModelReadFromDiskInBytes, const FString& InModelFullFilePath, const ENeuralDeviceType InDeviceType);
+//		TArray<bool>& OutAreInputTensorSizesVariable, const TArray<uint8>& InModelReadFromFileInBytes, const FString& InModelFullFilePath, const ENeuralDeviceType InDeviceType);
 //#endif //WITH_EDITOR
 
 	static bool Load(TSharedPtr<FImplBackEndUEAndORT>& InOutImplBackEndUEAndORT, TArray<FNeuralTensor>& OutInputTensors, TArray<FNeuralTensor>& OutOutputTensors,
-		TArray<bool>& OutAreInputTensorSizesVariable, const TArray<uint8>& InModelReadFromDiskInBytes, const FString& InModelFullFilePath, const ENeuralDeviceType InDeviceType);
+		TArray<bool>& OutAreInputTensorSizesVariable, const TArray<uint8>& InModelReadFromFileInBytes, const FString& InModelFullFilePath, const ENeuralDeviceType InDeviceType);
 	
 	void Run(const ENeuralNetworkSynchronousMode InSynchronousMode, const ENeuralDeviceType InInputDeviceType, const ENeuralDeviceType InOutputDeviceType);
 
diff --git a/Engine/Plugins/Experimental/NeuralNetworkInference/Source/NeuralNetworkInference/Public/NeuralNetwork.h b/Engine/Plugins/Experimental/NeuralNetworkInference/Source/NeuralNetworkInference/Public/NeuralNetwork.h
--- a/Engine/Plugins/Experimental/NeuralNetworkInference/Source/NeuralNetworkInference/Public/NeuralNetwork.h
+++ b/Engine/Plugins/Experimental/NeuralNetworkInference/Source/NeuralNetworkInference/Public/NeuralNetwork.h
@@ -74,17 +74,22 @@ public:
 	virtual ~UNeuralNetwork();
 
 	/**
-	 * Editor-only function
 	 * It loads the desired network graph definition and weights given an input ONNX file path.
 	 * @param InModelFilePath can either be a full path or a relative path with respect to the Game project.
 	 * @return Whether the network was successfully loaded.
 	 */
-#if WITH_EDITOR
 	bool Load(const FString& InModelFilePath);
-#endif // WITH_EDITOR
 
 	/**
-	 * It loads the desired network graph definition and weights given an input UNeuralNetwork UAsset.
+	 * It loads the desired network graph definition and weights given an input ONNX file path that has been read as a TArray<uint8> buffer.
+	 * @see FFileHelper::LoadFileToArray for an example of how to read a file into a TArray<uint8>.
+	 * @param InModelReadFromFileInBytes will be moved for performance reasons.
+	 * @return Whether the network was successfully loaded.
+	 */
+	bool Load(TArray<uint8>& InModelReadFromFileInBytes);
+
+	/**
+	 * It loads the desired network graph definition and weights internally saved on this UNeuralNetwork instance.
 	 * @return Whether the network was successfully loaded.
 	 */
 	bool Load();
@@ -212,7 +217,7 @@ private:
 	bool bIsLoaded;
 
 	UPROPERTY()
-	TArray<uint8> ModelReadFromDiskInBytes;
+	TArray<uint8> ModelReadFromFileInBytes;
 
 	/**
 	 * If BackEnd != Auto, BackEndForCurrentPlatform will be equal to BackEnd.
