NeuralNetworkInference (NNI) plugin:
- Bug fixed: Neural tensors have the same name than its network.
- FNeuralTensor: Added additional constructor to simplify things and add more flexibility: FNeuralTensor(const FString& InName, const ENeuralTensorTypeGPU InTensorTypeGPU)

#rnx

diff --git a/Engine/Plugins/Experimental/NeuralNetworkInference/Source/NeuralNetworkInference/Private/NeuralNetworkImplBackEndUEAndORT.cpp b/Engine/Plugins/Experimental/NeuralNetworkInference/Source/NeuralNetworkInference/Private/NeuralNetworkImplBackEndUEAndORT.cpp
--- a/Engine/Plugins/Experimental/NeuralNetworkInference/Source/NeuralNetworkInference/Private/NeuralNetworkImplBackEndUEAndORT.cpp
+++ b/Engine/Plugins/Experimental/NeuralNetworkInference/Source/NeuralNetworkInference/Private/NeuralNetworkImplBackEndUEAndORT.cpp
@@ -366,11 +366,34 @@ void UNeuralNetwork::FImplBackEndUEAndORT::SetTensorsFromNetwork(TArray<FNeuralT
 	TArray<const char*>& TensorNames = (bIsInput ? InputTensorNames : OutputTensorNames);
 	Swap(TensorNames, InTensorNames);
 
-	// Pre-allocate each variable
-	if (OutTensors.Num() != TensorNumber)
+	// Note: Switching from/to CPU to/from GPU would cause the FNeuralTensors to be re-initialized. We need to avoid that. For that, we will only re-allocate the tensors...
+	// - If bAreTensorsAlreadyCreatedWithRightNames == false, meaning tensors had not been created until now for this network.
+	// - And if the existing tensors have the right size, given that SetNumUninitialized() only re-allocates them if their size has changed.
+
+	// Fill bAreTensorsAlreadyCreatedWithRightNames - Check if tensors already created with the right names
+	bool bAreTensorsAlreadyCreatedWithRightNames = (OutTensors.Num() == TensorNames.Num());
+	if (bAreTensorsAlreadyCreatedWithRightNames)
+	{
+		for (int32 TensorIndex = 0; TensorIndex < TensorNumber; ++TensorIndex)
+		{
+			if (OutTensors[TensorIndex].GetName() != ANSI_TO_TCHAR(TensorNames[TensorIndex]))
+			{
+				bAreTensorsAlreadyCreatedWithRightNames = false;
+				break;
+			}
+		}
+	}
+
+	// Assign name to each input/output tensor
+	if (!bAreTensorsAlreadyCreatedWithRightNames)
 	{
-		OutTensors.SetNum(TensorNumber);
+		OutTensors.Empty();
+		for (const char* const TensorName : TensorNames)
+		{
+			OutTensors.Emplace(FNeuralTensor(ANSI_TO_TCHAR(TensorName)));
+		}
 	}
+	ensureMsgf(OutTensors.Num() == TensorNumber, TEXT("OutTensors.Num() == TensorNumber failed, %d != %d."), OutTensors.Num(), TensorNumber);
 
 	// Config each TensorIndex
 	TArray<Ort::Value>& OrtTensors = (bIsInput ? InputOrtTensors : OutputOrtTensors);
@@ -380,7 +403,7 @@ void UNeuralNetwork::FImplBackEndUEAndORT::SetTensorsFromNetwork(TArray<FNeuralT
 		{
 			OrtTensors.Emplace(Ort::Value(nullptr));
 		}
-		// Pre-allocate TArray
+		// Pre-allocate TArray (if size is different)
 		OutTensors[TensorIndex].SetNumUninitialized(InSizes[TensorIndex], InTensorDataTypes[TensorIndex]);
 		// Link tensor with ORT blob
 		LinkTensorToONNXRuntime(OutTensors, OrtTensors, *AllocatorInfo, TensorIndex);
diff --git a/Engine/Plugins/Experimental/NeuralNetworkInference/Source/NeuralNetworkInference/Private/NeuralTensor.cpp b/Engine/Plugins/Experimental/NeuralNetworkInference/Source/NeuralNetworkInference/Private/NeuralTensor.cpp
--- a/Engine/Plugins/Experimental/NeuralNetworkInference/Source/NeuralNetworkInference/Private/NeuralTensor.cpp
+++ b/Engine/Plugins/Experimental/NeuralNetworkInference/Source/NeuralNetworkInference/Private/NeuralTensor.cpp
@@ -94,14 +94,9 @@ void FPrivateNeuralTensor::NDTensorIndexesPlus1(TArray<int32>& InOutImageAreaInd
 
 
 
-/* FNeuralTensor structors
+/* FNeuralTensor costructors
  *****************************************************************************/
 
-FNeuralTensor::FNeuralTensor(const ENeuralDataType InDataType, const int64 InVolume, const FString& InName, const ENeuralTensorTypeGPU InTensorTypeGPU)
-	: FNeuralTensor(InDataType, InVolume > 0 ? TArray<int64>({ InVolume }) : TArray<int64>({}), InName, InTensorTypeGPU)
-{
-}
-
 FNeuralTensor::FNeuralTensor(const ENeuralDataType InDataType, const TArray<int64>& InSizes, const FString& InName, const ENeuralTensorTypeGPU InTensorTypeGPU)
 	: Name(InName)
 	, DataType(ENeuralDataType::None)
@@ -112,6 +107,21 @@ FNeuralTensor::FNeuralTensor(const ENeuralDataType InDataType, const TArray<int6
 	SetNumUninitialized(InSizes, InDataType);
 }
 
+FNeuralTensor::FNeuralTensor(const ENeuralDataType InDataType, const int64 InVolume, const FString& InName, const ENeuralTensorTypeGPU InTensorTypeGPU)
+	: FNeuralTensor(InDataType, InVolume > 0 ? TArray<int64>({ InVolume }) : TArray<int64>({}), InName, InTensorTypeGPU)
+{
+}
+
+FNeuralTensor::FNeuralTensor(const FString& InName, const ENeuralTensorTypeGPU InTensorTypeGPU)
+	: FNeuralTensor(ENeuralDataType::None, TArray<int64>({}), InName, InTensorTypeGPU)
+{
+}
+
+
+
+/* FNeuralTensor private costructor
+ *****************************************************************************/
+
 FNeuralTensor::FNeuralTensor(const ENeuralDataType InDataType, const void* const InValues, const int64 InSizeOfT, const int64 InValueNum, const TArray<int64>& InSizes, const FString& InName, const ENeuralTensorTypeGPU InTensorTypeGPU)
 	: FNeuralTensor(InDataType, InSizes, InName, InTensorTypeGPU)
 {
@@ -127,6 +137,11 @@ FNeuralTensor::FNeuralTensor(const ENeuralDataType InDataType, const void* const
 	}
 }
 
+
+
+/* FNeuralTensor copy/move structors/assignment
+ *****************************************************************************/
+
 FNeuralTensor::FNeuralTensor(const FNeuralTensor& InTensor)
 	: FNeuralTensor(InTensor.DataType, InTensor.Sizes, InTensor.Name)
 {
diff --git a/Engine/Plugins/Experimental/NeuralNetworkInference/Source/NeuralNetworkInference/Public/NeuralTensor.h b/Engine/Plugins/Experimental/NeuralNetworkInference/Source/NeuralNetworkInference/Public/NeuralTensor.h
--- a/Engine/Plugins/Experimental/NeuralNetworkInference/Source/NeuralNetworkInference/Public/NeuralTensor.h
+++ b/Engine/Plugins/Experimental/NeuralNetworkInference/Source/NeuralNetworkInference/Public/NeuralTensor.h
@@ -54,8 +54,9 @@ public:
 	 * @param InSizes Set to empty (or omit argument) if memory allocation is not required or the final size is unknown.
 	 * @param InName Used for GPU debugging and the ToString() function.
 	*/
-	explicit FNeuralTensor(const ENeuralDataType InDataType, const int64 InVolume, const FString& InName = TEXT("FNeuralTensor"), const ENeuralTensorTypeGPU InTensorTypeGPU = ENeuralTensorTypeGPU::Generic);
 	explicit FNeuralTensor(const ENeuralDataType InDataType = ENeuralDataType::None, const TArray<int64>& InSizes = TArray<int64>(), const FString& InName = TEXT("FNeuralTensor"), const ENeuralTensorTypeGPU InTensorTypeGPU = ENeuralTensorTypeGPU::Generic);
+	explicit FNeuralTensor(const ENeuralDataType InDataType, const int64 InVolume, const FString& InName = TEXT("FNeuralTensor"), const ENeuralTensorTypeGPU InTensorTypeGPU = ENeuralTensorTypeGPU::Generic);
+	explicit FNeuralTensor(const FString& InName, const ENeuralTensorTypeGPU InTensorTypeGPU = ENeuralTensorTypeGPU::Generic);
 
 	/**
 	 * Performance-wise, this constructor makes a deep copy of the data (not optimal). For maximum speed, use the other constructors + GetData()/GetDataCasted().
