NeuralNetworkInference (NNI) plugin:
- Changes in UNN & UEAndORT back end to make easier later integration with UEOnly back end.
  - Input/OutputTensors moved from UNN into ImplBackEndUEAndORT.
  - Input/OutputTensors moved no longer TPROPERTY's.
  - Removed GetInputTensors()/GetOutputTensors(). Instead, added GetInputTensorNum()/GetOutputTensorNum()
- With above changes done, UEOnly can finally fully work with UNN:
  - TensorManager used for UEOnly back end when asking for a particular input/output tensor.
- NNIExampleCharacter.cpp now runs even if WITH_UE_AND_ORT_SUPPORT is disabled (but only with compatible nets, cloth & MLDef).
- Added sanity checks to all UNN classes to read/modify input/output tensors.

#rnx

diff --git a/Engine/Plugins/Experimental/NeuralNetworkInference/Source/NeuralNetworkInference/Private/NeuralNetwork.cpp b/Engine/Plugins/Experimental/NeuralNetworkInference/Source/NeuralNetworkInference/Private/NeuralNetwork.cpp
--- a/Engine/Plugins/Experimental/NeuralNetworkInference/Source/NeuralNetworkInference/Private/NeuralNetwork.cpp
+++ b/Engine/Plugins/Experimental/NeuralNetworkInference/Source/NeuralNetworkInference/Private/NeuralNetwork.cpp
@@ -126,7 +126,7 @@ bool UNeuralNetwork::Load()
 	// UEAndORT
 	if (BackEndForCurrentPlatform == ENeuralBackEnd::UEAndORT)
 	{
-		bIsLoaded = UNeuralNetwork::FImplBackEndUEAndORT::Load(ImplBackEndUEAndORT, InputTensors, OutputTensors, AreInputTensorSizesVariable, ModelReadFromFileInBytes, ModelFullFilePath, GetDeviceType());
+		bIsLoaded = UNeuralNetwork::FImplBackEndUEAndORT::Load(ImplBackEndUEAndORT, AreInputTensorSizesVariable, ModelReadFromFileInBytes, ModelFullFilePath, GetDeviceType());
 	}
 	// UEOnly
 	else if (BackEndForCurrentPlatform == ENeuralBackEnd::UEOnly)
@@ -228,41 +228,46 @@ void UNeuralNetwork::SetBackEnd(const ENeuralBackEnd InBackEnd)
 
 const FNeuralTensor& UNeuralNetwork::GetInputTensor(const int32 InTensorIndex) const
 {
+	// Sanity check
+	checkf(bIsLoaded, TEXT("Call UNeuralNetwork::Load() to load a model first."));
+
 	// UEAndORT
 	if (BackEndForCurrentPlatform == ENeuralBackEnd::UEAndORT)
 	{
-		return InputTensors[InTensorIndex];
+		return ImplBackEndUEAndORT->InputTensors[InTensorIndex];
 	}
 
 	// UEOnly
 	else if (BackEndForCurrentPlatform == ENeuralBackEnd::UEOnly)
 	{
-		const int32 TensorIndex = ImplBackEndUEOnly->TensorManager.GetInputIndexes()[InTensorIndex];
-		return ImplBackEndUEOnly->TensorManager.GetTensors()[TensorIndex];
+		FNeuralTensorManager& TensorManager = ImplBackEndUEOnly->TensorManager;
+		return TensorManager.GetTensors()[TensorManager.GetInputIndexes()[InTensorIndex]];
 	}
 
 	// Unknown
 	checkf(false, TEXT("UNeuralNetwork::GetInputTensor(): Unknown [BackEnd,BackEndForCurrentPlatform] = [%d,%d]."), (int32)BackEnd, (int32)BackEndForCurrentPlatform);
-	return InputTensors[InTensorIndex];
-}
-
-const TArray<FNeuralTensor>& UNeuralNetwork::GetInputTensors() const
-{
-	return InputTensors;
+	return ImplBackEndUEAndORT->InputTensors[InTensorIndex];
 }
 
 void UNeuralNetwork::SetInputFromArrayCopy(const TArray<float>& InArray, const int32 InTensorIndex)
 {
+	// Sanity check
+	if (!bIsLoaded)
+	{
+		UE_LOG(LogNeuralNetworkInference, Warning, TEXT("UNeuralNetwork::SetInputFromArrayCopy(): Call UNeuralNetwork::Load() to load a model first."));
+	}
+
 	// UEAndORT
 	if (BackEndForCurrentPlatform == ENeuralBackEnd::UEAndORT)
 	{
-		InputTensors[InTensorIndex].SetFromArrayCopy(InArray);
+		ImplBackEndUEAndORT->InputTensors[InTensorIndex].SetFromArrayCopy(InArray);
 	}
 
 	// UEOnly
 	else if (BackEndForCurrentPlatform == ENeuralBackEnd::UEOnly)
 	{
-UE_LOG(LogNeuralNetworkInference, Warning, TEXT("UNeuralNetwork::SetInputFromArrayCopy(): UEOnly not implemented yet."));
+		FNeuralTensorManager& TensorManager = ImplBackEndUEOnly->TensorManager;
+		TensorManager.GetTensorsMutable()[TensorManager.GetInputIndexes()[InTensorIndex]].SetFromArrayCopy(InArray);
 	}
 
 	// Unknown
@@ -274,16 +279,24 @@ UE_LOG(LogNeuralNetworkInference, Warning, TEXT("UNeuralNetwork::SetInputFromArr
 
 void* UNeuralNetwork::GetInputDataPointerMutable(const int32 InTensorIndex)
 {
+	// Sanity check
+	if (!bIsLoaded)
+	{
+		UE_LOG(LogNeuralNetworkInference, Warning, TEXT("UNeuralNetwork::GetInputDataPointerMutable(): Call UNeuralNetwork::Load() to load a model first."));
+		return nullptr;
+	}
+
 	// UEAndORT
 	if (BackEndForCurrentPlatform == ENeuralBackEnd::UEAndORT)
 	{
-		return InputTensors[InTensorIndex].GetData(); // Or ImplBackEndUEAndORT->InputOrtTensors[InTensorIndex].GetTensorMutableData<float>();
+		return ImplBackEndUEAndORT->InputTensors[InTensorIndex].GetData(); // Or ImplBackEndUEAndORT->InputOrtTensors[InTensorIndex].GetTensorMutableData<float>();
 	}
 
 	// UEOnly
 	else if (BackEndForCurrentPlatform == ENeuralBackEnd::UEOnly)
 	{
-UE_LOG(LogNeuralNetworkInference, Warning, TEXT("UNeuralNetwork::GetInputDataPointerMutable(): UEOnly not implemented yet."));
+		FNeuralTensorManager& TensorManager = ImplBackEndUEOnly->TensorManager;
+		return TensorManager.GetTensorsMutable()[TensorManager.GetInputIndexes()[InTensorIndex]].GetData();
 	}
 
 	// Unknown
@@ -291,38 +304,89 @@ UE_LOG(LogNeuralNetworkInference, Warning, TEXT("UNeuralNetwork::GetInputDataPoi
 	return nullptr;
 }
 
+int64 UNeuralNetwork::GetInputTensorNumber() const
+{
+	// Sanity check
+	if (!bIsLoaded)
+	{
+		UE_LOG(LogNeuralNetworkInference, Warning, TEXT("UNeuralNetwork::GetInputTensorNumber(): Call UNeuralNetwork::Load() to load a model first."));
+		return -1;
+	}
+
+	// UEAndORT
+	if (BackEndForCurrentPlatform == ENeuralBackEnd::UEAndORT)
+	{
+		return ImplBackEndUEAndORT->InputTensors.Num();
+	}
+
+	// UEOnly
+	else if (BackEndForCurrentPlatform == ENeuralBackEnd::UEOnly)
+	{
+		return ImplBackEndUEOnly->TensorManager.GetInputIndexes().Num();
+	}
+
+	// Unknown
+	checkf(false, TEXT("UNeuralNetwork::GetInputTensorNumber(): Unknown [BackEnd,BackEndForCurrentPlatform] = [%d,%d]."), (int32)BackEnd, (int32)BackEndForCurrentPlatform);
+	return -1;
+}
+
 const FNeuralTensor& UNeuralNetwork::GetOutputTensor(const int32 InTensorIndex) const
 {
+	// Sanity check
+	checkf(bIsLoaded, TEXT("Call UNeuralNetwork::Load() to load a model first."));
+
 	// UEAndORT
 	if (BackEndForCurrentPlatform == ENeuralBackEnd::UEAndORT)
 	{
-		return OutputTensors[InTensorIndex];
+		return ImplBackEndUEAndORT->OutputTensors[InTensorIndex];
 	}
 
 	// UEOnly
 	else if (BackEndForCurrentPlatform == ENeuralBackEnd::UEOnly)
 	{
-UE_LOG(LogNeuralNetworkInference, Warning, TEXT("UNeuralNetwork::GetOutputTensor(): UEOnly not implemented yet."));
+		FNeuralTensorManager& TensorManager = ImplBackEndUEOnly->TensorManager;
+		return TensorManager.GetTensors()[TensorManager.GetOutputIndexes()[InTensorIndex]];
 	}
 
 	// Unknown
 	checkf(false, TEXT("UNeuralNetwork::GetOutputTensor(): Unknown [BackEnd,BackEndForCurrentPlatform] = [%d,%d]."), (int32)BackEnd, (int32)BackEndForCurrentPlatform);
-	return OutputTensors[InTensorIndex];
+	return ImplBackEndUEAndORT->OutputTensors[InTensorIndex];
 }
 
-const TArray<FNeuralTensor>& UNeuralNetwork::GetOutputTensors() const
+int64 UNeuralNetwork::GetOutputTensorNumber() const
 {
-	return OutputTensors;
+	// Sanity check
+	if (!bIsLoaded)
+	{
+		UE_LOG(LogNeuralNetworkInference, Warning, TEXT("UNeuralNetwork::GetOutputTensorNumber(): Call UNeuralNetwork::Load() to load a model first."));
+		return -1;
+	}
+
+	// UEAndORT
+	if (BackEndForCurrentPlatform == ENeuralBackEnd::UEAndORT)
+	{
+		return ImplBackEndUEAndORT->OutputTensors.Num();
+	}
+
+	// UEOnly
+	else if (BackEndForCurrentPlatform == ENeuralBackEnd::UEOnly)
+	{
+		return ImplBackEndUEOnly->TensorManager.GetOutputIndexes().Num();
+	}
+
+	// Unknown
+	checkf(false, TEXT("UNeuralNetwork::GetOutputTensorNumber(): Unknown [BackEnd,BackEndForCurrentPlatform] = [%d,%d]."), (int32)BackEnd, (int32)BackEndForCurrentPlatform);
+	return -1;
 }
 
 void UNeuralNetwork::Run()
 {
 	DECLARE_SCOPE_CYCLE_COUNTER(TEXT("UNeuralNetwork_Run"), STAT_UNeuralNetwork_Run, STATGROUP_MachineLearning);
 
-	// Sanity checks
+	// Sanity check
 	if (!bIsLoaded)
 	{
-		UE_LOG(LogNeuralNetworkInference, Warning, TEXT("UNeuralNetwork::Run(): No architecture has been loaded yet. Run() will not work until IsLoaded() returns true."));
+		UE_LOG(LogNeuralNetworkInference, Warning, TEXT("UNeuralNetwork::Run(): Call UNeuralNetwork::Load() to load a model first."));
 		return;
 	}
 
diff --git a/Engine/Plugins/Experimental/NeuralNetworkInference/Source/NeuralNetworkInference/Private/NeuralNetworkImplBackEndUEAndORT.cpp b/Engine/Plugins/Experimental/NeuralNetworkInference/Source/NeuralNetworkInference/Private/NeuralNetworkImplBackEndUEAndORT.cpp
--- a/Engine/Plugins/Experimental/NeuralNetworkInference/Source/NeuralNetworkInference/Private/NeuralNetworkImplBackEndUEAndORT.cpp
+++ b/Engine/Plugins/Experimental/NeuralNetworkInference/Source/NeuralNetworkInference/Private/NeuralNetworkImplBackEndUEAndORT.cpp
@@ -9,8 +9,7 @@
  *****************************************************************************/
 
 //#if WITH_EDITOR
-//bool UNeuralNetwork::FImplBackEndUEAndORT::LoadFile(TSharedPtr<FImplBackEndUEAndORT>& InOutImplBackEndUEAndORT, TArray<FNeuralTensor>& OutInputTensors, TArray<FNeuralTensor>& OutOutputTensors,
-//	TArray<bool>& OutAreInputTensorSizesVariable, const TArray<uint8>& InModelReadFromFileInBytes, const FString& InModelFullFilePath, const ENeuralDeviceType InDeviceType)
+//bool UNeuralNetwork::FImplBackEndUEAndORT::LoadFile(TSharedPtr<FImplBackEndUEAndORT>& InOutImplBackEndUEAndORT, TArray<bool>& OutAreInputTensorSizesVariable, const TArray<uint8>& InModelReadFromFileInBytes, const FString& InModelFullFilePath, const ENeuralDeviceType InDeviceType)
 //{
 //	// Initialize and configure InOutImplBackEndUEAndORT
 //	const FRedirectCoutAndCerrToUeLog RedirectCoutAndCerrToUeLog;
@@ -77,8 +76,7 @@
 //}
 //#endif //WITH_EDITOR
 
-bool UNeuralNetwork::FImplBackEndUEAndORT::Load(TSharedPtr<FImplBackEndUEAndORT>& InOutImplBackEndUEAndORT, TArray<FNeuralTensor>& OutInputTensors, TArray<FNeuralTensor>& OutOutputTensors,
-	TArray<bool>& OutAreInputTensorSizesVariable, const TArray<uint8>& InModelReadFromFileInBytes, const FString& InModelFullFilePath, const ENeuralDeviceType InDeviceType)
+bool UNeuralNetwork::FImplBackEndUEAndORT::Load(TSharedPtr<FImplBackEndUEAndORT>& InOutImplBackEndUEAndORT, TArray<bool>& OutAreInputTensorSizesVariable, const TArray<uint8>& InModelReadFromFileInBytes, const FString& InModelFullFilePath, const ENeuralDeviceType InDeviceType)
 {
 #ifdef WITH_UE_AND_ORT_SUPPORT
 #if WITH_EDITOR
@@ -107,8 +105,8 @@ bool UNeuralNetwork::FImplBackEndUEAndORT::Load(TSharedPtr<FImplBackEndUEAndORT>
 			return false;
 		}
 
-		InOutImplBackEndUEAndORT->ConfigureTensors(OutInputTensors, &OutAreInputTensorSizesVariable);
-		InOutImplBackEndUEAndORT->ConfigureTensors(OutOutputTensors);
+		InOutImplBackEndUEAndORT->ConfigureTensors(InOutImplBackEndUEAndORT->InputTensors, &OutAreInputTensorSizesVariable);
+		InOutImplBackEndUEAndORT->ConfigureTensors(InOutImplBackEndUEAndORT->OutputTensors);
 
 		return true;
 	}
diff --git a/Engine/Plugins/Experimental/NeuralNetworkInference/Source/NeuralNetworkInference/Private/NeuralNetworkImplBackEndUEAndORT.h b/Engine/Plugins/Experimental/NeuralNetworkInference/Source/NeuralNetworkInference/Private/NeuralNetworkImplBackEndUEAndORT.h
--- a/Engine/Plugins/Experimental/NeuralNetworkInference/Source/NeuralNetworkInference/Private/NeuralNetworkImplBackEndUEAndORT.h
+++ b/Engine/Plugins/Experimental/NeuralNetworkInference/Source/NeuralNetworkInference/Private/NeuralNetworkImplBackEndUEAndORT.h
@@ -28,13 +28,18 @@ NNI_THIRD_PARTY_INCLUDES_END
 struct UNeuralNetwork::FImplBackEndUEAndORT
 {
 public:
+	/**
+	 * InputTensors and OutputTensors represent the input and output TArray<FNeuralTensor> of the network, respectively.
+	 */
+	TArray<FNeuralTensor> InputTensors;
+	TArray<FNeuralTensor> OutputTensors;
 #ifdef WITH_UE_AND_ORT_SUPPORT
-	// Network-related variables
+	/** Network-related variables */
 	TUniquePtr<Ort::Env> Environment;
 	TUniquePtr<Ort::Session> Session;
 	TUniquePtr<Ort::AllocatorWithDefaultOptions> Allocator;
 	TUniquePtr<Ort::SessionOptions> SessionOptions;
-	// Tensor-related variables
+	/** Tensor-related variables */
 	TUniquePtr<Ort::MemoryInfo> AllocatorInfo; /* Memory allocator information */
 	TArray<Ort::Value> InputOrtTensors; /* Actual ONNXRuntime tensors */
 	TArray<const char*> InputTensorNames; /* Tensor names */
@@ -43,12 +48,10 @@ public:
 #endif //WITH_UE_AND_ORT_SUPPORT
 
 //#if WITH_EDITOR
-//	static bool LoadFile(TSharedPtr<FImplBackEndUEAndORT>& InOutImplBackEndUEAndORT, TArray<FNeuralTensor>& OutInputTensors, TArray<FNeuralTensor>& OutOutputTensors,
-//		TArray<bool>& OutAreInputTensorSizesVariable, const TArray<uint8>& InModelReadFromFileInBytes, const FString& InModelFullFilePath, const ENeuralDeviceType InDeviceType);
+//	static bool LoadFile(TSharedPtr<FImplBackEndUEAndORT>& InOutImplBackEndUEAndORT, TArray<bool>& OutAreInputTensorSizesVariable, const TArray<uint8>& InModelReadFromFileInBytes, const FString& InModelFullFilePath, const ENeuralDeviceType InDeviceType);
 //#endif //WITH_EDITOR
 
-	static bool Load(TSharedPtr<FImplBackEndUEAndORT>& InOutImplBackEndUEAndORT, TArray<FNeuralTensor>& OutInputTensors, TArray<FNeuralTensor>& OutOutputTensors,
-		TArray<bool>& OutAreInputTensorSizesVariable, const TArray<uint8>& InModelReadFromFileInBytes, const FString& InModelFullFilePath, const ENeuralDeviceType InDeviceType);
+	static bool Load(TSharedPtr<FImplBackEndUEAndORT>& InOutImplBackEndUEAndORT, TArray<bool>& OutAreInputTensorSizesVariable, const TArray<uint8>& InModelReadFromFileInBytes, const FString& InModelFullFilePath, const ENeuralDeviceType InDeviceType);
 	
 	void Run(const ENeuralNetworkSynchronousMode InSynchronousMode, const ENeuralDeviceType InInputDeviceType, const ENeuralDeviceType InOutputDeviceType);
 
diff --git a/Engine/Plugins/Experimental/NeuralNetworkInference/Source/NeuralNetworkInference/Private/NeuralNetworkImplBackEndUEOnly.cpp b/Engine/Plugins/Experimental/NeuralNetworkInference/Source/NeuralNetworkInference/Private/NeuralNetworkImplBackEndUEOnly.cpp
--- a/Engine/Plugins/Experimental/NeuralNetworkInference/Source/NeuralNetworkInference/Private/NeuralNetworkImplBackEndUEOnly.cpp
+++ b/Engine/Plugins/Experimental/NeuralNetworkInference/Source/NeuralNetworkInference/Private/NeuralNetworkImplBackEndUEOnly.cpp
@@ -40,7 +40,9 @@ bool UNeuralNetwork::FImplBackEndUEOnly::Load(TSharedPtr<FImplBackEndUEOnly>& In
 		UE_LOG(LogNeuralNetworkInference, Warning, TEXT("UNeuralNetwork::FImplBackEndUEOnly::Load(): TensorManager could not be loaded from %s."), *InModelFilePath);
 		return false;
 	}
+	// Everything worked out, return true
 	return true;
+
 #else //WITH_EDITOR
 	UE_LOG(LogNeuralNetworkInference, Warning, TEXT("UNeuralNetwork::FImplBackEndUEOnly::Load(): Only implemented for Editor mode."));
 	return false;
@@ -64,76 +66,6 @@ bool UNeuralNetwork::FImplBackEndUEOnly::Load(TSharedPtr<FImplBackEndUEOnly>& In
 //	return (InOutImplBackEndUEOnly->Operators.Num() > 0 && InOutImplBackEndUEOnly->TensorManager.IsLoaded());
 //}
 
-//const TArray<FNeuralTensor>& UNeuralNetwork::FImplBackEndUEOnly::GetTensors() const
-//{
-//	return TensorManager.GetTensors();
-//}
-//
-//void UNeuralNetwork::FImplBackEndUEOnly::SetInputFromTensorCopy(const FNeuralTensor& InTensorMap)
-//{
-//	return TensorManager.SetInputFromTensorCopy(InTensorMap);
-//}
-//
-//void UNeuralNetwork::FImplBackEndUEOnly::SetInputFromTensorMapCopy(const TMap<FString, FNeuralTensor>& InTensorMap)
-//{
-//	return TensorManager.SetInputFromTensorMapCopy(InTensorMap);
-//}
-//
-//TMap<FString, void*> UNeuralNetwork::FImplBackEndUEOnly::CreateInputDataPointersMutable()
-//{
-//	return TensorManager.CreateInputDataPointersMutable();
-//}
-//
-////FRDGBufferUAVRef UNeuralNetwork::FImplBackEndUEOnly::GetInputBufferUAVRef()
-////{
-////	return TensorManager.GetInputBufferUAVRef();
-////}
-////
-////TMap<FString, FRDGBufferUAVRef> UNeuralNetwork::FImplBackEndUEOnly::CreateInputBufferUAVRefs()
-////{
-////	return TensorManager.CreateInputBufferUAVRefs();
-////}
-//
-//const FNeuralTensor& UNeuralNetwork::FImplBackEndUEOnly::GetInputTensor() const
-//{
-//	return TensorManager.GetInputTensor();
-//}
-//
-//const TMap<FString, int32>& UNeuralNetwork::FImplBackEndUEOnly::GetInputNameIndexMap() const
-//{
-//	return TensorManager.GetInputNameIndexMap();
-//}
-//
-//const FNeuralTensor& UNeuralNetwork::FImplBackEndUEOnly::GetOutputTensor() const
-//{
-//	return TensorManager.GetOutputTensor();
-//}
-//
-//const TMap<FString, int32>& UNeuralNetwork::FImplBackEndUEOnly::GetOutputNameIndexMap() const
-//{
-//	return TensorManager.GetOutputNameIndexMap();
-//}
-//
-//const FRDGBufferSRVRef UNeuralNetwork::FImplBackEndUEOnly::GetOutputBufferSRVRef() const
-//{
-//	return TensorManager.GetOutputBufferSRVRef();
-//}
-//
-//TMap<FString, const FRDGBufferSRVRef> UNeuralNetwork::FImplBackEndUEOnly::CreateOutputBufferSRVRefs() const
-//{
-//	return TensorManager.CreateOutputBufferSRVRefs();
-//}
-//
-//TMap<FString, FNeuralTensor> UNeuralNetwork::FImplBackEndUEOnly::CreateInputTensorMap() const
-//{
-//	return TensorManager.CreateInputTensorMap();
-//}
-//
-//TMap<FString, FNeuralTensor> UNeuralNetwork::FImplBackEndUEOnly::CreateOutputTensorMap() const
-//{
-//	return TensorManager.CreateOutputTensorMap();
-//}
-
 void UNeuralNetwork::FImplBackEndUEOnly::Run(FOnAsyncRunCompleted& InOutOnAsyncRunCompletedDelegate, const ENeuralNetworkSynchronousMode InSynchronousMode, const ENeuralDeviceType InDeviceType, const ENeuralDeviceType InInputDeviceType, const ENeuralDeviceType InOutputDeviceType)
 {
 	// Run UNeuralNetworkLegacy
diff --git a/Engine/Plugins/Experimental/NeuralNetworkInference/Source/NeuralNetworkInference/Public/NeuralNetwork.h b/Engine/Plugins/Experimental/NeuralNetworkInference/Source/NeuralNetworkInference/Public/NeuralNetwork.h
--- a/Engine/Plugins/Experimental/NeuralNetworkInference/Source/NeuralNetworkInference/Public/NeuralNetwork.h
+++ b/Engine/Plugins/Experimental/NeuralNetworkInference/Source/NeuralNetworkInference/Public/NeuralNetwork.h
@@ -143,15 +143,15 @@ public:
 	 * These functions either take a TArray as input, return a modificable void* to fill the data, or return a constant FNeuralTensor(s) to see input properties (e.g., size or dimensions).
 	 */
 	const FNeuralTensor& GetInputTensor(const int32 InTensorIndex = 0) const;
-	const TArray<FNeuralTensor>& GetInputTensors() const;
 	void SetInputFromArrayCopy(const TArray<float>& InArray, const int32 InTensorIndex = 0);
 	void* GetInputDataPointerMutable(const int32 InTensorIndex = 0);
+	int64 GetInputTensorNumber() const;
 
 	/**
 	 * Functions to get output. The returned FNeuralTensor(s) are constant to prevent the user from modifying the tensor properties (e.g., size or dimensions).
 	 */
 	const FNeuralTensor& GetOutputTensor(const int32 InTensorIndex = 0) const;
-	const TArray<FNeuralTensor>& GetOutputTensors() const;
+	int64 GetOutputTensorNumber() const;
 
 	/**
 	 * Run() executes the forward pass on the current UNeuralNetwork given the current input FDeprecatedNeuralTensor(s), which were previously filled with
@@ -200,25 +200,15 @@ protected:
 	UPROPERTY(VisibleAnywhere, Category = "Neural Network Inference")
 	FString ModelFullFilePath;
 
-	/**
-	 * InputTensors and OutputTensors represent the input and output TArray<FNeuralTensor> of the network, respectively.
-	 */
-	UPROPERTY(VisibleAnywhere, Category = "Neural Network Inference")
-	TArray<FNeuralTensor> InputTensors;
-
-	UPROPERTY(VisibleAnywhere, Category = "Neural Network Inference")
-	TArray<FNeuralTensor> OutputTensors;
-
-	/** Whether some of the FNeuralTensor of InputTensor have flexible/variable dimensions. */
-	UPROPERTY(VisibleAnywhere, Category = "Neural Network Inference")
-	TArray<bool> AreInputTensorSizesVariable;
-
 private:
 	bool bIsLoaded;
 
 	UPROPERTY()
 	TArray<uint8> ModelReadFromFileInBytes;
 
+	/** Whether some of the FNeuralTensor of InputTensor have flexible/variable dimensions. */
+	TArray<bool> AreInputTensorSizesVariable;
+
 	/**
 	 * If BackEnd != Auto, BackEndForCurrentPlatform will be equal to BackEnd.
 	 * Otherwise, BackEndForCurrentPlatform will be set to the optimal BackEnd given the current platform.
diff --git a/Engine/Plugins/Experimental/NeuralNetworkInference/Source/NeuralNetworkInferenceQA/Private/ModelUnitTester.cpp b/Engine/Plugins/Experimental/NeuralNetworkInference/Source/NeuralNetworkInferenceQA/Private/ModelUnitTester.cpp
--- a/Engine/Plugins/Experimental/NeuralNetworkInference/Source/NeuralNetworkInferenceQA/Private/ModelUnitTester.cpp
+++ b/Engine/Plugins/Experimental/NeuralNetworkInference/Source/NeuralNetworkInferenceQA/Private/ModelUnitTester.cpp
@@ -12,20 +12,30 @@
 
 void FModelUnitTester::GlobalTest(const FString& InProjectContentDir, const FString& InModelZooRelativeDirectory)
 {
+	// Model names, input values, and number of repetitions for profiling
 #ifdef WITH_UE_AND_ORT_SUPPORT
-	// Model load, accuracy, and speed test
 	const TArray<FString> ModelNames({ TEXT("MLRigDeformer"), TEXT("cloth_network"), TEXT("HS"), TEXT("RL") });
 	const TArray<float> InputArrayValues({ 1.f, 0.f, -1.f, 100.f, -100.f, 0.5f, -0.5f }); // This one can be shorter than CPU/GPUGroundTruths
+	const TArray<int32> CPURepetitions({ 1000, 1000,  50, 1000 });
+	const TArray<int32> GPURepetitions({ 1000, 1000, 100, 1000 });
+#else
+	UE_LOG(LogNeuralNetworkInferenceQA, Display, TEXT("--------------- Some FModelUnitTester tests skipped (only if WITH_UE_AND_ORT_SUPPORT)."));
+UE_LOG(LogNeuralNetworkInferenceQA, Display, TEXT("--------------- FModelUnitTester test skipped (only if WITH_UE_AND_ORT_SUPPORT)."));
+return;
+	const TArray<FString> ModelNames({ TEXT("MLRigDeformer"), TEXT("cloth_network") });
+	const TArray<float> InputArrayValues({ 1.f, 0.f });
+	const TArray<int32> CPURepetitions({ 100, 0 });
+	const TArray<int32> GPURepetitions({ 100, 100 });
+#endif //WITH_UE_AND_ORT_SUPPORT
+
+	// Ground truths
 	const TArray<TArray<double>> CPUGroundTruths({ {3.728547, 0.008774, 4.595651, 212.193216, 742.434561, 4.250668, 4.717748}, {0.042571, 0.023693, 0.015783, 13.100505, 8.050994, 0.028807, 0.016387},
 		{138.372906, 126.753839, 127.287254, 130.316062, 127.303424, 124.800896, 126.546051}, {0.488662, 0.472437, 0.478862, 0.522685, 0.038322, 0.480848, 0.483821}});
 	const TArray<TArray<double>> GPUGroundTruths({ {3.728547, 0.008774, 4.595651, 212.193208, 742.434578, 4.250668, 4.717748}, {0.042571, 0.023693, 0.015783, 13.100504, 8.050994, 0.028807, 0.016387},
 		{138.373184, 126.754100, 127.287398, 130.316194, 127.303495, 124.801134, 126.5462530}, {0.488662, 0.472437, 0.478862, 0.522685, 0.038322, 0.480848, 0.483821}});
-	const TArray<int32> CPURepetitions({ 1000, 1000,  50, 1000 });
-	const TArray<int32> GPURepetitions({ 1000, 1000, 100, 1000 });
+
+	// Run tests
 	ModelLoadAccuracyAndSpeedTests(InProjectContentDir, InModelZooRelativeDirectory, ModelNames, InputArrayValues, CPUGroundTruths, GPUGroundTruths, CPURepetitions, GPURepetitions);
-#else
-	UE_LOG(LogNeuralNetworkInferenceQA, Display, TEXT("--------------- FModelUnitTester test skipped (only if WITH_UE_AND_ORT_SUPPORT)."));
-#endif //WITH_UE_AND_ORT_SUPPORT
 }
 
 
@@ -56,12 +66,12 @@ void FModelUnitTester::ModelLoadAccuracyAndSpeedTests(const FString& InProjectCo
 		}
 		// Input debugging
 		UE_LOG(LogNeuralNetworkInferenceQA, Display, TEXT("--------------- %s - Input/Output"), *ModelName);
-		for (int32 TensorIndex = 0; TensorIndex < Network->GetInputTensors().Num(); ++TensorIndex)
+		for (int32 TensorIndex = 0; TensorIndex < Network->GetInputTensorNumber(); ++TensorIndex)
 		{
 			UE_LOG(LogNeuralNetworkInferenceQA, Display, TEXT("InputTensor[%d] = %s."), TensorIndex, *Network->GetInputTensor(TensorIndex).GetName());
 		}
 		// Output debugging
-		for (int32 TensorIndex = 0; TensorIndex < Network->GetOutputTensors().Num(); ++TensorIndex)
+		for (int32 TensorIndex = 0; TensorIndex < Network->GetOutputTensorNumber(); ++TensorIndex)
 		{
 			UE_LOG(LogNeuralNetworkInferenceQA, Display, TEXT("OutputTensor[%d] = %s."), TensorIndex, *Network->GetOutputTensor(TensorIndex).GetName());
 		}
@@ -191,9 +201,9 @@ void FModelUnitTester::ModelAccuracyTest(UNeuralNetwork* InOutNetwork, const TAr
 		const double FastCPUGPUAvgL1NormDiff = FMath::Abs((CPUAvgL1Norm - GPUAvgL1Norm)) * RelativeCoefficient * 1e6;
 		const double FastCPUAvgL1NormDiff = FMath::Abs(CPUAvgL1Norm - CPUGroundTruth) / FMath::Max(1., FMath::Min(CPUAvgL1Norm, CPUGroundTruth)) * 1e7;
 		const double FastGPUAvgL1NormDiff = FMath::Abs(GPUAvgL1Norm - GPUGroundTruth) / FMath::Max(1., FMath::Min(GPUAvgL1Norm, GPUGroundTruth)) * 1e7;
-		UE_LOG(LogNeuralNetworkInferenceQA, Display, TEXT("InputNorm = %f, OutputNormCPU = %f, OutputNormGT = %f, CPUAvgL1Norm = %f, GPUAvgL1Norm = %f,"
+		UE_LOG(LogNeuralNetworkInferenceQA, Display, TEXT("InputNorm = %f, OutputNormCPU = %f, OutputNormGPU = %f, OutputNormGT = %f, CPUAvgL1Norm = %f, GPUAvgL1Norm = %f,"
 			"\tCPUGPUAvgL1NormDiff = %fe-3, FastCPUGPUAvgL1NormDiff = %fe-6, FastCPUAvgL1NormDiff = %fe-7, FastGPUAvgL1NormDiff = %fe-7 (1e-7 is roughly the precision for float)."),
-			GetAveragedL1Norm(InputArray), CPUAvgL1Norm, CPUGroundTruth, CPUAvgL1Norm, GPUAvgL1Norm, CPUGPUAvgL1NormDiff, FastCPUGPUAvgL1NormDiff, FastCPUAvgL1NormDiff, FastGPUAvgL1NormDiff);
+			GetAveragedL1Norm(InputArray), CPUAvgL1Norm, GPUAvgL1Norm, CPUGroundTruth, CPUAvgL1Norm, GPUAvgL1Norm, CPUGPUAvgL1NormDiff, FastCPUGPUAvgL1NormDiff, FastCPUAvgL1NormDiff, FastGPUAvgL1NormDiff);
 		const TArray<int64>& InputSizes = InOutNetwork->GetInputTensor().GetSizes();
 		const TArray<int64>& OutputSizes = InOutNetwork->GetOutputTensor().GetSizes();
 		const int64 MaxNumberElementsToDisplay = 100;
