NeuralNetworkInference (NNI) plugin:
- Automation: Added "AddWarning" if any test failed.
- Added return bool to all tests to know if they passed or not.

#rnx

diff --git a/Engine/Plugins/Experimental/NeuralNetworkInference/Source/NeuralNetworkInferenceQA/Private/LegacyModelUnitTester.cpp b/Engine/Plugins/Experimental/NeuralNetworkInference/Source/NeuralNetworkInferenceQA/Private/LegacyModelUnitTester.cpp
--- a/Engine/Plugins/Experimental/NeuralNetworkInference/Source/NeuralNetworkInferenceQA/Private/LegacyModelUnitTester.cpp
+++ b/Engine/Plugins/Experimental/NeuralNetworkInference/Source/NeuralNetworkInferenceQA/Private/LegacyModelUnitTester.cpp
@@ -1,7 +1,7 @@
 // Copyright Epic Games, Inc. All Rights Reserved.
 
 #include "LegacyModelUnitTester.h"
-#include "LegacyONNXTester.h"
+#include "ModelProtoFileReaderTester.h"
 #include "NeuralNetworkInferenceQATimer.h"
 #include "NeuralNetworkInferenceQAUtils.h"
 #include "Misc/Paths.h"
@@ -11,7 +11,7 @@
 /* FLegacyModelUnitTester static public functions
  *****************************************************************************/
 
-void FLegacyModelUnitTester::GlobalTest(const FString& InProjectContentDir, const FString& InModelZooRelativeDirectory)
+bool FLegacyModelUnitTester::GlobalTest(const FString& InProjectContentDir, const FString& InModelZooRelativeDirectory)
 {
 	// Model load, accuracy, and speed test
 	const TArray<FString> ModelNames({ TEXT("MLRigDeformer")/*, TEXT("cloth_network")*/ });
@@ -20,7 +20,7 @@ void FLegacyModelUnitTester::GlobalTest(const FString& InProjectContentDir, cons
 	const TArray<TArray<double>> GPUGroundTruths({ {3.728547, 0.008774, 4.595651, 212.193208, 742.434578, 4.250668, 4.717748}, {0.042571, 0.023693, 0.015783, 13.100504, 8.050994, 0.028807, 0.016387} });
 	const TArray<int32> CPURepetitions({ 10, 0 }); // Used to be 100
 	const TArray<int32> GPURepetitions({ 10, 10 }); // Used to be 100
-	ModelLoadAccuracyAndSpeedTests(InProjectContentDir, InModelZooRelativeDirectory, ModelNames, InputArrayValues, CPUGroundTruths, GPUGroundTruths, CPURepetitions, GPURepetitions);
+	return ModelLoadAccuracyAndSpeedTests(InProjectContentDir, InModelZooRelativeDirectory, ModelNames, InputArrayValues, CPUGroundTruths, GPUGroundTruths, CPURepetitions, GPURepetitions);
 }
 
 
@@ -28,10 +28,12 @@ void FLegacyModelUnitTester::GlobalTest(const FString& InProjectContentDir, cons
 /* FLegacyModelUnitTester static private functions
  *****************************************************************************/
 
-void FLegacyModelUnitTester::ModelLoadAccuracyAndSpeedTests(const FString& InProjectContentDir, const FString& InModelZooRelativeDirectory, const TArray<FString>& InModelNames,
+bool FLegacyModelUnitTester::ModelLoadAccuracyAndSpeedTests(const FString& InProjectContentDir, const FString& InModelZooRelativeDirectory, const TArray<FString>& InModelNames,
 	const TArray<float>& InInputArrayValues, const TArray<TArray<double>>& InCPUGroundTruths, const TArray<TArray<double>>& InGPUGroundTruths, const TArray<int32>& InCPURepetitions,
 	const TArray<int32>& InGPURepetitions)
 {
+	bool bDidGlobalTestPassed = true;
+
 	const FString ModelZooDirectory = InProjectContentDir / InModelZooRelativeDirectory;
 
 	// Test ONNX Reader
@@ -39,10 +41,10 @@ void FLegacyModelUnitTester::ModelLoadAccuracyAndSpeedTests(const FString& InPro
 	{
 		UE_LOG(LogNeuralNetworkInferenceQA, Display, TEXT("--------------- %s - OTXT IO"), *ModelName);
 		const FString OtxtModelFilePath = GetOtxtModelFilePath(ModelZooDirectory, ModelName);
-		FLegacyONNXTester::ONNXReadNetworkTest(OtxtModelFilePath);
+		bDidGlobalTestPassed &= FModelProtoFileReaderTester::Test(OtxtModelFilePath);
 		UE_LOG(LogNeuralNetworkInferenceQA, Display, TEXT("--------------- %s - ONNX IO"), *ModelName);
 		const FString ONNXModelFilePath = GetONNXModelFilePath(ModelZooDirectory, ModelName);
-		FLegacyONNXTester::ONNXReadNetworkTest(ONNXModelFilePath);
+		bDidGlobalTestPassed &= FModelProtoFileReaderTester::Test(ONNXModelFilePath);
 	}
 	UE_LOG(LogNeuralNetworkInferenceQA, Display, TEXT("---------------------------------------------------------------------------------------------------------------------------------"));
 
@@ -55,11 +57,11 @@ void FLegacyModelUnitTester::ModelLoadAccuracyAndSpeedTests(const FString& InPro
 
 		const FString ONNXModelFilePath = GetONNXModelFilePath(ModelZooDirectory, ModelName);
 		UE_LOG(LogNeuralNetworkInferenceQA, Display, TEXT("--------------- %s - Legacy Network ONNX Load and Run - %s"), *ModelName, *ONNXModelFilePath);
-		ModelAccuracyTest(NetworkONNXLoadTest(ONNXModelFilePath), InInputArrayValues, CPUGroundTruths, GPUGroundTruths);
+		bDidGlobalTestPassed &= ModelAccuracyTest(NetworkONNXLoadTest(ONNXModelFilePath), InInputArrayValues, CPUGroundTruths, GPUGroundTruths);
 
 		UE_LOG(LogNeuralNetworkInferenceQA, Display, TEXT("--------------- %s - Legacy Network Uasset Load and Run"), *ModelName);
 		const FString UAssetModelFilePath = GetUAssetModelFilePath(ModelName, InModelZooRelativeDirectory);
-		ModelAccuracyTest(NetworkUassetLoadTest(UAssetModelFilePath), InInputArrayValues, CPUGroundTruths, GPUGroundTruths);
+		bDidGlobalTestPassed &= ModelAccuracyTest(NetworkUassetLoadTest(UAssetModelFilePath), InInputArrayValues, CPUGroundTruths, GPUGroundTruths);
 
 		UE_LOG(LogNeuralNetworkInferenceQA, Display, TEXT("---------------------------------------------------------------------------------------------------------------------------------"));
 	}
@@ -70,8 +72,10 @@ void FLegacyModelUnitTester::ModelLoadAccuracyAndSpeedTests(const FString& InPro
 		const FString& ModelName = InModelNames[ModelIndex];
 		UE_LOG(LogNeuralNetworkInferenceQA, Display, TEXT("--------------- %s - Legacy Network UAsset Speed Profiling"), *ModelName);
 		const FString UAssetModelFilePath = GetUAssetModelFilePath(ModelName, InModelZooRelativeDirectory);
-		ModelSpeedTest(UAssetModelFilePath, InCPURepetitions[ModelIndex], InGPURepetitions[ModelIndex]);
+		bDidGlobalTestPassed &= ModelSpeedTest(UAssetModelFilePath, InCPURepetitions[ModelIndex], InGPURepetitions[ModelIndex]);
 	}
+
+	return true;
 }
 
 FString FLegacyModelUnitTester::GetONNXModelFilePath(const FString& ModelZooDirectory, const FString& InModelName)
@@ -128,11 +132,12 @@ UNeuralNetworkLegacy* FLegacyModelUnitTester::NetworkONNXLoadTest(const FString&
 #endif //WITH_EDITOR
 }
 
-void FLegacyModelUnitTester::ModelAccuracyTest(UNeuralNetworkLegacy* InOutNetwork, const TArray<float>& InInputArrayValues, const TArray<double>& InCPUGroundTruths, const TArray<double>& InGPUGroundTruths)
+bool FLegacyModelUnitTester::ModelAccuracyTest(UNeuralNetworkLegacy* InOutNetwork, const TArray<float>& InInputArrayValues, const TArray<double>& InCPUGroundTruths, const TArray<double>& InGPUGroundTruths)
 {
+	// Sanity check
 	if (!InOutNetwork)
 	{
-		return;
+		return false;
 	}
 	// Find NetworkSize
 	const int64 NetworkSize = InOutNetwork->GetInputTensor().Num();
@@ -198,6 +203,8 @@ void FLegacyModelUnitTester::ModelAccuracyTest(UNeuralNetworkLegacy* InOutNetwor
 		ensureMsgf(FMath::IsFinite(FastGPUAvgL1NormDiff) && FastGPUAvgL1NormDiff < 5, TEXT("FastGPUAvgL1NormDiff (%fe-7) < 5e-7 failed (~5 times the float precision).\nOutputArrayGPU = %s"),
 			FastGPUAvgL1NormDiff, *FNeuralTensor(OutputArrayGPU, OutputSizes).ToString(MaxNumberElementsToDisplay));
 	}
+	// Test successful
+	return true;
 }
 
 #define LEGACY_MODEL_UNIT_TESTER_SPEED_TEST(OutCopyTimeInMilliSeconds, OutNetworkTimeInMilliSeconds, OutEmptyNetworkTimeInMilliSeconds, InRepetitions, bInIsGPU) \
@@ -239,7 +246,7 @@ void FLegacyModelUnitTester::ModelAccuracyTest(UNeuralNetworkLegacy* InOutNetwor
 	} \
 	const float OutNetworkTimeInMilliSeconds = Timer.Toc() / InRepetitions - OutCopyTimeInMilliSeconds - (bInIsGPU ? OutEmptyNetworkTimeInMilliSeconds : 0)
 
-void FLegacyModelUnitTester::ModelSpeedTest(const FString& InUAssetPath, const int32 InCPURepetitions, const int32 InGPURepetitions)
+bool FLegacyModelUnitTester::ModelSpeedTest(const FString& InUAssetPath, const int32 InCPURepetitions, const int32 InGPURepetitions)
 {
 	// Load Network
 	UNeuralNetworkLegacy* InOutNetwork = NetworkUassetLoadTest(InUAssetPath);
@@ -247,7 +254,7 @@ void FLegacyModelUnitTester::ModelSpeedTest(const FString& InUAssetPath, const i
 	if (!InOutNetwork)
 	{
 		ensureMsgf(false, TEXT("InOutNetwork was a nullptr."));
-		return;
+		return false;
 	}
 	// Needed variables
 	const int64 NetworkSize = InOutNetwork->GetInputTensor().Num();
@@ -271,6 +278,8 @@ void FLegacyModelUnitTester::ModelSpeedTest(const FString& InUAssetPath, const i
 		TEXT("Forward pass speed profiling (TensorCopy, optionally EmptyForward(), and Forward()): CPUx1 = %f+%f msec, CPUx%d = %f+%f msec, GPUx1 = %f+%f+%f msec, GPUx%d = %f+%f+%f msec."),
 		CPUCopyTimer1,                        CPUNetworkTimer1, InCPURepetitions, CPUCopyTimer,                       CPUNetworkTimer,
 		GPUCopyTimer1, GPUEmptyNetworkTimer1, GPUNetworkTimer1, InGPURepetitions, GPUCopyTimer, GPUEmptyNetworkTimer, GPUNetworkTimer);
+	// Test successful
+	return true;
 }
 
 double FLegacyModelUnitTester::GetAveragedL1Norm(const TArray<float>& InArray)
diff --git a/Engine/Plugins/Experimental/NeuralNetworkInference/Source/NeuralNetworkInferenceQA/Private/LegacyModelUnitTester.h b/Engine/Plugins/Experimental/NeuralNetworkInference/Source/NeuralNetworkInferenceQA/Private/LegacyModelUnitTester.h
--- a/Engine/Plugins/Experimental/NeuralNetworkInference/Source/NeuralNetworkInferenceQA/Private/LegacyModelUnitTester.h
+++ b/Engine/Plugins/Experimental/NeuralNetworkInference/Source/NeuralNetworkInferenceQA/Private/LegacyModelUnitTester.h
@@ -8,13 +8,13 @@
 class FLegacyModelUnitTester
 {
 public:
-	static void GlobalTest(const FString& InProjectContentDir, const FString& InModelZooRelativeDirectory);
+	static bool GlobalTest(const FString& InProjectContentDir, const FString& InModelZooRelativeDirectory);
 
 private:
 	/**
 	 * It runs a full model test on the desired model.
 	 */
-	static void ModelLoadAccuracyAndSpeedTests(const FString& InProjectContentDir, const FString& InModelZooRelativeDirectory, const TArray<FString>& InModelNames,
+	static bool ModelLoadAccuracyAndSpeedTests(const FString& InProjectContentDir, const FString& InModelZooRelativeDirectory, const TArray<FString>& InModelNames,
 		const TArray<float>& InInputArrayValues, const TArray<TArray<double>>& InCPUGroundTruths, const TArray<TArray<double>>& InGPUGroundTruths,
 		const TArray<int32>& InCPURepetitions, const TArray<int32>& InGPURepetitions);
 	/**
@@ -25,9 +25,9 @@ private:
 	static FString GetUAssetModelFilePath(const FString& InModelName, const FString& InModelZooRelativeDirectory);
 	static UNeuralNetworkLegacy* NetworkUassetLoadTest(const FString& InUAssetPath);
 	static UNeuralNetworkLegacy* NetworkONNXLoadTest(const FString& InONNXFilePath);
-	static void ModelAccuracyTest(UNeuralNetworkLegacy* InOutNetwork, const TArray<float>& InInputArrayValues, const TArray<double>& InCPUGroundTruths,
+	static bool ModelAccuracyTest(UNeuralNetworkLegacy* InOutNetwork, const TArray<float>& InInputArrayValues, const TArray<double>& InCPUGroundTruths,
 		const TArray<double>& InGPUGroundTruths);
-	static void ModelSpeedTest(const FString& InUAssetPath, const int32 InCPURepetitions, const int32 InGPURepetitions);
+	static bool ModelSpeedTest(const FString& InUAssetPath, const int32 InCPURepetitions, const int32 InGPURepetitions);
 	static double GetAveragedL1Norm(const TArray<float>& InArray);
 	static double GetAveragedL1NormDiff(const TArray<float>& InArray1, const TArray<float>& InArray2);
 };
diff --git a/Engine/Plugins/Experimental/NeuralNetworkInference/Source/NeuralNetworkInferenceQA/Private/LegacyONNXTester.cpp b/Engine/Plugins/Experimental/NeuralNetworkInference/Source/NeuralNetworkInferenceQA/Private/LegacyONNXTester.cpp
deleted file mode 100644
diff --git a/Engine/Plugins/Experimental/NeuralNetworkInference/Source/NeuralNetworkInferenceQA/Private/LegacyONNXTester.h b/Engine/Plugins/Experimental/NeuralNetworkInference/Source/NeuralNetworkInferenceQA/Private/LegacyONNXTester.h
deleted file mode 100644
diff --git a/Engine/Plugins/Experimental/NeuralNetworkInference/Source/NeuralNetworkInferenceQA/Private/ModelProtoFileReaderTester.cpp b/Engine/Plugins/Experimental/NeuralNetworkInference/Source/NeuralNetworkInferenceQA/Private/ModelProtoFileReaderTester.cpp
new file mode 100644
--- /dev/null
+++ b/Engine/Plugins/Experimental/NeuralNetworkInference/Source/NeuralNetworkInferenceQA/Private/ModelProtoFileReaderTester.cpp
@@ -0,0 +1,29 @@
+// Copyright Epic Games, Inc. All Rights Reserved.
+
+#include "ModelProtoFileReaderTester.h"
+#include "NeuralNetworkInferenceQAUtils.h"
+
+#if WITH_EDITOR
+#include "ModelProtoFileReader.h"
+#endif //WITH_EDITOR
+
+
+
+/* FModelProtoFileReaderTester public functions
+ *****************************************************************************/
+
+bool FModelProtoFileReaderTester::Test(const FString& InONNXModelFileName)
+{
+	UE_LOG(LogNeuralNetworkInferenceQA, Display, TEXT("-------------------- Read ONNX Network And Test"));
+
+#if WITH_EDITOR
+	FModelProto ModelProto;
+	FModelProtoFileReader::ReadModelProtoFromFile(ModelProto, InONNXModelFileName);
+	UE_LOG(LogNeuralNetworkInferenceQA, Display, TEXT("ModelProto:\n%s"), *ModelProto.ToString());
+	ensureMsgf(ModelProto.IsLoaded(), TEXT("FModelProtoFileReaderTester::ONNXReadNetworkTest() failed, FModelProto could not be read from InONNXModelFileName: %s."), *InONNXModelFileName);
+	return ModelProto.IsLoaded();
+#else //WITH_EDITOR
+	UE_LOG(LogNeuralNetworkInferenceQA, Display, TEXT("ONNXReadNetworkTest test skipped (only in Editor)."));
+	return true;
+#endif //WITH_EDITOR
+}
diff --git a/Engine/Plugins/Experimental/NeuralNetworkInference/Source/NeuralNetworkInferenceQA/Private/ModelProtoFileReaderTester.h b/Engine/Plugins/Experimental/NeuralNetworkInference/Source/NeuralNetworkInferenceQA/Private/ModelProtoFileReaderTester.h
new file mode 100644
--- /dev/null
+++ b/Engine/Plugins/Experimental/NeuralNetworkInference/Source/NeuralNetworkInferenceQA/Private/ModelProtoFileReaderTester.h
@@ -0,0 +1,11 @@
+// Copyright Epic Games, Inc. All Rights Reserved.
+
+#pragma once
+
+#include "CoreMinimal.h"
+
+class FModelProtoFileReaderTester
+{
+public:
+	static bool Test(const FString& InONNXModelFileName);
+};
diff --git a/Engine/Plugins/Experimental/NeuralNetworkInference/Source/NeuralNetworkInferenceQA/Private/ModelUnitTester.cpp b/Engine/Plugins/Experimental/NeuralNetworkInference/Source/NeuralNetworkInferenceQA/Private/ModelUnitTester.cpp
--- a/Engine/Plugins/Experimental/NeuralNetworkInference/Source/NeuralNetworkInferenceQA/Private/ModelUnitTester.cpp
+++ b/Engine/Plugins/Experimental/NeuralNetworkInference/Source/NeuralNetworkInferenceQA/Private/ModelUnitTester.cpp
@@ -10,7 +10,7 @@
 /* FModelUnitTester static public functions
  *****************************************************************************/
 
-void FModelUnitTester::GlobalTest(const FString& InProjectContentDir, const FString& InModelZooRelativeDirectory)
+bool FModelUnitTester::GlobalTest(const FString& InProjectContentDir, const FString& InModelZooRelativeDirectory)
 {
 	// Model names, input values, and number of repetitions for profiling
 #ifdef WITH_UE_AND_ORT_SUPPORT
@@ -20,12 +20,12 @@ void FModelUnitTester::GlobalTest(const FString& InProjectContentDir, const FStr
 	const TArray<int32> GPURepetitions({ 1000, 1000, 100, 1000 });
 #else
 	UE_LOG(LogNeuralNetworkInferenceQA, Display, TEXT("--------------- Some FModelUnitTester tests skipped (only if WITH_UE_AND_ORT_SUPPORT)."));
-UE_LOG(LogNeuralNetworkInferenceQA, Display, TEXT("--------------- FModelUnitTester test skipped (only if WITH_UE_AND_ORT_SUPPORT)."));
-return;
 	const TArray<FString> ModelNames({ TEXT("MLRigDeformer"), TEXT("cloth_network") });
 	const TArray<float> InputArrayValues({ 1.f, 0.f });
 	const TArray<int32> CPURepetitions({ 100, 0 });
 	const TArray<int32> GPURepetitions({ 100, 100 });
+UE_LOG(LogNeuralNetworkInferenceQA, Display, TEXT("--------------- FModelUnitTester test skipped (only if WITH_UE_AND_ORT_SUPPORT)."));
+return true;
 #endif //WITH_UE_AND_ORT_SUPPORT
 
 	// Ground truths
@@ -35,7 +35,7 @@ return;
 		{138.373184, 126.754100, 127.287398, 130.316194, 127.303495, 124.801134, 126.5462530}, {0.488662, 0.472437, 0.478862, 0.522685, 0.038322, 0.480848, 0.483821}});
 
 	// Run tests
-	ModelLoadAccuracyAndSpeedTests(InProjectContentDir, InModelZooRelativeDirectory, ModelNames, InputArrayValues, CPUGroundTruths, GPUGroundTruths, CPURepetitions, GPURepetitions);
+	return ModelLoadAccuracyAndSpeedTests(InProjectContentDir, InModelZooRelativeDirectory, ModelNames, InputArrayValues, CPUGroundTruths, GPUGroundTruths, CPURepetitions, GPURepetitions);
 }
 
 
@@ -43,10 +43,12 @@ return;
 /* FModelUnitTester static private functions
  *****************************************************************************/
 
-void FModelUnitTester::ModelLoadAccuracyAndSpeedTests(const FString& InProjectContentDir, const FString& InModelZooRelativeDirectory, const TArray<FString>& InModelNames,
+bool FModelUnitTester::ModelLoadAccuracyAndSpeedTests(const FString& InProjectContentDir, const FString& InModelZooRelativeDirectory, const TArray<FString>& InModelNames,
 	const TArray<float>& InInputArrayValues, const TArray<TArray<double>>& InCPUGroundTruths, const TArray<TArray<double>>& InGPUGroundTruths, const TArray<int32>& InCPURepetitions,
 	const TArray<int32>& InGPURepetitions)
 {
+	bool bDidGlobalTestPassed = true;
+
 	const FString ModelZooDirectory = InProjectContentDir / InModelZooRelativeDirectory;
 
 	// Test OTXT/UAsset accuracy
@@ -62,7 +64,7 @@ void FModelUnitTester::ModelLoadAccuracyAndSpeedTests(const FString& InProjectCo
 		if (!Network)
 		{
 			ensureMsgf(false, TEXT("UNeuralNetwork could not be loaded from UAssetModelFilePath %s."), *UAssetModelFilePath);
-			return;
+			return false;
 		}
 		// Input debugging
 		UE_LOG(LogNeuralNetworkInferenceQA, Display, TEXT("--------------- %s - Input/Output"), *ModelName);
@@ -75,7 +77,7 @@ void FModelUnitTester::ModelLoadAccuracyAndSpeedTests(const FString& InProjectCo
 		{
 			UE_LOG(LogNeuralNetworkInferenceQA, Display, TEXT("OutputTensor[%d] = %s."), TensorIndex, *Network->GetOutputTensor(TensorIndex).GetName());
 		}
-		ModelAccuracyTest(Network, InInputArrayValues, CPUGroundTruths, GPUGroundTruths);
+		bDidGlobalTestPassed &= ModelAccuracyTest(Network, InInputArrayValues, CPUGroundTruths, GPUGroundTruths);
 
 		UE_LOG(LogNeuralNetworkInferenceQA, Display, TEXT("---------------------------------------------------------------------------------------------------------------------------------"));
 
@@ -90,7 +92,7 @@ if (ONNXOrORTIndex == 1)
 	break;
 }
 			UE_LOG(LogNeuralNetworkInferenceQA, Display, TEXT("-------------------- %s - Network %s Load and Run - %s"), *ModelName, *ModelType, *ModelFilePath);
-			ModelAccuracyTest(NetworkONNXOrORTLoadTest(ModelFilePath), InInputArrayValues, CPUGroundTruths, GPUGroundTruths);
+			bDidGlobalTestPassed &= ModelAccuracyTest(NetworkONNXOrORTLoadTest(ModelFilePath), InInputArrayValues, CPUGroundTruths, GPUGroundTruths);
 		}
 #else
 	UE_LOG(LogNeuralNetworkInferenceQA, Display, TEXT("-------------------- Skipped (only in Editor)."));
@@ -103,8 +105,10 @@ if (ONNXOrORTIndex == 1)
 		const FString& ModelName = InModelNames[ModelIndex];
 		UE_LOG(LogNeuralNetworkInferenceQA, Display, TEXT("--------------- %s - Network UAsset Speed Profiling"), *ModelName);
 		const FString UAssetModelFilePath = GetUAssetModelFilePath(ModelName, InModelZooRelativeDirectory);
-		ModelSpeedTest(UAssetModelFilePath, InCPURepetitions[ModelIndex], InGPURepetitions[ModelIndex]);
+		bDidGlobalTestPassed &= ModelSpeedTest(UAssetModelFilePath, InCPURepetitions[ModelIndex], InGPURepetitions[ModelIndex]);
 	}
+
+	return true;
 }
 
 FString FModelUnitTester::GetONNXModelFilePath(const FString& ModelZooDirectory, const FString& InModelName)
@@ -157,11 +161,12 @@ UNeuralNetwork* FModelUnitTester::NetworkONNXOrORTLoadTest(const FString& InMode
 }
 #endif //WITH_EDITOR
 
-void FModelUnitTester::ModelAccuracyTest(UNeuralNetwork* InOutNetwork, const TArray<float>& InInputArrayValues, const TArray<double>& InCPUGroundTruths, const TArray<double>& InGPUGroundTruths)
+bool FModelUnitTester::ModelAccuracyTest(UNeuralNetwork* InOutNetwork, const TArray<float>& InInputArrayValues, const TArray<double>& InCPUGroundTruths, const TArray<double>& InGPUGroundTruths)
 {
+	// Sanity check
 	if (!InOutNetwork)
 	{
-		return;
+		return false;
 	}
 	// Find NetworkSize
 	const int64 NetworkSize = InOutNetwork->GetInputTensor().Num();
@@ -179,6 +184,7 @@ void FModelUnitTester::ModelAccuracyTest(UNeuralNetwork* InOutNetwork, const TAr
 			InputArrays.Num(), InCPUGroundTruths.Num(), InGPUGroundTruths.Num());
 	}
 	// Run each input with CPU/GPU and compare with each other and with the ground truth
+	bool bDidGlobalTestPassed = true;
 	for (int32 Index = 0; Index < InputArrays.Num(); ++Index)
 	{
 		TArray<float>& InputArray = InputArrays[Index];
@@ -223,8 +229,12 @@ void FModelUnitTester::ModelAccuracyTest(UNeuralNetwork* InOutNetwork, const TAr
 			UE_LOG(LogNeuralNetworkInferenceQA, Display, TEXT("OutputArrayCPU = %s"), *FNeuralTensor(OutputArrayCPU, OutputSizes).ToString(MaxNumberElementsToDisplay));
 			UE_LOG(LogNeuralNetworkInferenceQA, Display, TEXT("OutputArrayGPU = %s"), *FNeuralTensor(OutputArrayGPU, OutputSizes).ToString(MaxNumberElementsToDisplay));
 			ensureMsgf(false, TEXT("At least 1 of the 4 CPU/GPU tests failed."));
+			return false;
 		}
+		bDidGlobalTestPassed &= !bDidSomeTestFailed;
 	}
+	// Test successful
+	return bDidGlobalTestPassed;
 }
 
 #define MODEL_UNIT_TESTER_SPEED_TEST(OutCopyTimeInMilliSeconds, OutNetworkTimeInMilliSeconds, InRepetitions, bInIsGPU) \
@@ -256,7 +266,7 @@ void FModelUnitTester::ModelAccuracyTest(UNeuralNetwork* InOutNetwork, const TAr
 	} \
 	const float OutNetworkTimeInMilliSeconds = Timer.Toc() / InRepetitions - OutCopyTimeInMilliSeconds
 
-void FModelUnitTester::ModelSpeedTest(const FString& InUAssetPath, const int32 InCPURepetitions, const int32 InGPURepetitions)
+bool FModelUnitTester::ModelSpeedTest(const FString& InUAssetPath, const int32 InCPURepetitions, const int32 InGPURepetitions)
 {
 	// Load Network
 	UNeuralNetwork* InOutNetwork = NetworkUassetLoadTest(InUAssetPath);
@@ -264,7 +274,7 @@ void FModelUnitTester::ModelSpeedTest(const FString& InUAssetPath, const int32 I
 	if (!InOutNetwork)
 	{
 		ensureMsgf(false, TEXT("InOutNetwork was a nullptr."));
-		return;
+		return false;
 	}
 	// Needed variables
 	const int64 NetworkSize = InOutNetwork->GetInputTensor().Num();
@@ -288,6 +298,8 @@ void FModelUnitTester::ModelSpeedTest(const FString& InUAssetPath, const int32 I
 		TEXT("Forward pass speed profiling (TensorCopy and Forward()): CPUx1 = %f+%f msec, CPUx%d = %f+%f msec, GPUx1 = %f+%f msec, GPUx%d = %f+%f msec."),
 		CPUCopyTimer1, CPUNetworkTimer1, InCPURepetitions, CPUCopyTimer, CPUNetworkTimer,
 		GPUCopyTimer1, GPUNetworkTimer1, InGPURepetitions, GPUCopyTimer, GPUNetworkTimer);
+	// Test successful
+	return true;
 }
 
 double FModelUnitTester::GetAveragedL1Norm(const TArray<float>& InArray)
diff --git a/Engine/Plugins/Experimental/NeuralNetworkInference/Source/NeuralNetworkInferenceQA/Private/ModelUnitTester.h b/Engine/Plugins/Experimental/NeuralNetworkInference/Source/NeuralNetworkInferenceQA/Private/ModelUnitTester.h
--- a/Engine/Plugins/Experimental/NeuralNetworkInference/Source/NeuralNetworkInferenceQA/Private/ModelUnitTester.h
+++ b/Engine/Plugins/Experimental/NeuralNetworkInference/Source/NeuralNetworkInferenceQA/Private/ModelUnitTester.h
@@ -8,13 +8,13 @@
 class FModelUnitTester
 {
 public:
-	static void GlobalTest(const FString& InProjectContentDir, const FString& InModelZooRelativeDirectory);
+	static bool GlobalTest(const FString& InProjectContentDir, const FString& InModelZooRelativeDirectory);
 
 private:
 	/**
 	 * It runs a full model test on the desired model.
 	 */
-	static void ModelLoadAccuracyAndSpeedTests(const FString& InProjectContentDir, const FString& InModelZooRelativeDirectory, const TArray<FString>& InModelNames,
+	static bool ModelLoadAccuracyAndSpeedTests(const FString& InProjectContentDir, const FString& InModelZooRelativeDirectory, const TArray<FString>& InModelNames,
 		const TArray<float>& InInputArrayValues, const TArray<TArray<double>>& InCPUGroundTruths, const TArray<TArray<double>>& InGPUGroundTruths,
 		const TArray<int32>& InCPURepetitions, const TArray<int32>& InGPURepetitions);
 	/**
@@ -27,9 +27,9 @@ private:
 #if WITH_EDITOR
 	static UNeuralNetwork* NetworkONNXOrORTLoadTest(const FString& InModelFilePath);
 #endif //WITH_EDITOR
-	static void ModelAccuracyTest(UNeuralNetwork* InOutNetwork, const TArray<float>& InInputArrayValues, const TArray<double>& InCPUGroundTruths,
+	static bool ModelAccuracyTest(UNeuralNetwork* InOutNetwork, const TArray<float>& InInputArrayValues, const TArray<double>& InCPUGroundTruths,
 		const TArray<double>& InGPUGroundTruths);
-	static void ModelSpeedTest(const FString& InUAssetPath, const int32 InCPURepetitions, const int32 InGPURepetitions);
+	static bool ModelSpeedTest(const FString& InUAssetPath, const int32 InCPURepetitions, const int32 InGPURepetitions);
 	static double GetAveragedL1Norm(const TArray<float>& InArray);
 	static double GetAveragedL1NormDiff(const TArray<float>& InArray1, const TArray<float>& InArray2);
 };
diff --git a/Engine/Plugins/Experimental/NeuralNetworkInference/Source/NeuralNetworkInferenceQA/Private/NeuralNetworkInferenceQA.cpp b/Engine/Plugins/Experimental/NeuralNetworkInference/Source/NeuralNetworkInferenceQA/Private/NeuralNetworkInferenceQA.cpp
--- a/Engine/Plugins/Experimental/NeuralNetworkInference/Source/NeuralNetworkInferenceQA/Private/NeuralNetworkInferenceQA.cpp
+++ b/Engine/Plugins/Experimental/NeuralNetworkInference/Source/NeuralNetworkInferenceQA/Private/NeuralNetworkInferenceQA.cpp
@@ -10,13 +10,15 @@
 /* UNeuralNetworkInferenceQA static functions
  *****************************************************************************/
 
-void UNeuralNetworkInferenceQA::UnitTesting()
+bool UNeuralNetworkInferenceQA::UnitTesting()
 {
 	const FString ProjectContentDir = FPaths::ProjectContentDir();
 	const FString MachineLearningTestsRelativeDirectory = TEXT("Tests/MachineLearning/");
 	const FString ModelZooRelativeDirectory = MachineLearningTestsRelativeDirectory / TEXT("Models/"); // Eg "Tests/MachineLearning/Models/"
 	const FString UnitTestRelativeDirectory = MachineLearningTestsRelativeDirectory / TEXT("UnitTests/"); // Eg "Tests/MachineLearning/UnitTests/"
-	FUnitTester::GlobalTest(ProjectContentDir, ModelZooRelativeDirectory, UnitTestRelativeDirectory);
+	const bool bDidGlobalTestPassed = FUnitTester::GlobalTest(ProjectContentDir, ModelZooRelativeDirectory, UnitTestRelativeDirectory);
+	ensureMsgf(bDidGlobalTestPassed, TEXT("UNeuralNetworkInferenceQA::UnitTesting() returned false. See log above."));
+	return bDidGlobalTestPassed;
 }
 
 
@@ -32,8 +34,10 @@ IMPLEMENT_SIMPLE_AUTOMATION_TEST(FNeuralNetworkInferenceTest, "System.Engine.Mac
 
 bool FNeuralNetworkInferenceTest::RunTest(const FString& Parameters)
 {
-	UNeuralNetworkInferenceQA::UnitTesting();
-UE_LOG(LogNeuralNetworkInferenceQA, Warning, TEXT("FNeuralNetworkInferenceTest::RunTest(): Warning with parameters = %s."), *Parameters);
+	if (!UNeuralNetworkInferenceQA::UnitTesting())
+	{
+		AddError(TEXT("FNeuralNetworkInferenceTest::RunTest(): UNeuralNetworkInferenceQA::UnitTesting() returned false. See log above.")); // AddWarning/AddError would both work here
+	}
 	return true;
 }
 
diff --git a/Engine/Plugins/Experimental/NeuralNetworkInference/Source/NeuralNetworkInferenceQA/Private/OperatorUnitTester.cpp b/Engine/Plugins/Experimental/NeuralNetworkInference/Source/NeuralNetworkInferenceQA/Private/OperatorUnitTester.cpp
--- a/Engine/Plugins/Experimental/NeuralNetworkInference/Source/NeuralNetworkInferenceQA/Private/OperatorUnitTester.cpp
+++ b/Engine/Plugins/Experimental/NeuralNetworkInference/Source/NeuralNetworkInferenceQA/Private/OperatorUnitTester.cpp
@@ -14,7 +14,7 @@
 /* FOperatorUnitTester static public functions
  *****************************************************************************/
 
-void FOperatorUnitTester::GlobalTest(const FString& InProjectContentDir, const FString& InUnitTestRelativeDirectory)
+bool FOperatorUnitTester::GlobalTest(const FString& InProjectContentDir, const FString& InUnitTestRelativeDirectory)
 {
 	const FString UnitTestDirectory = InProjectContentDir / InUnitTestRelativeDirectory;
 	const float ZeroThreshold = 0.f;
@@ -144,7 +144,7 @@ void FOperatorUnitTester::GlobalTest(const FString& InProjectContentDir, const F
 	if (!NeuralNetworkInferenceQAAsset)
 	{
 		ensureMsgf(false, TEXT("NeuralNetworkInferenceQAAsset could not be loaded!"));
-		return;
+		return false;
 	}
 	NeuralNetworkInferenceQAAsset->FindOrAddOperators(OperatorNames);
 	NeuralNetworkInferenceQAAsset->FlushNewTests();
@@ -208,10 +208,13 @@ void FOperatorUnitTester::GlobalTest(const FString& InProjectContentDir, const F
 	{
 		// Saving NeuralNetworkInferenceQAAsset
 		ensureMsgf(NeuralNetworkInferenceQAAsset->Save(), TEXT("NeuralNetworkInferenceQAAsset->Save() failed."));
-
 		// Error message
 		ensureMsgf(false, TEXT("NeuralNetworkInferenceQAAsset->CompareNewVsPreviousTests() failed. Updated NeuralNetworkInferenceQAAsset saved as %s. Check the previous warning messages."), *NeuralNetworkInferenceQAAsset->GetOutermost()->GetLoadedPath().GetLocalFullPath());
+		// Test failed
+		return false;
 	}
+	// Test successful
+	return true;
 }
 
 
diff --git a/Engine/Plugins/Experimental/NeuralNetworkInference/Source/NeuralNetworkInferenceQA/Private/OperatorUnitTester.h b/Engine/Plugins/Experimental/NeuralNetworkInference/Source/NeuralNetworkInferenceQA/Private/OperatorUnitTester.h
--- a/Engine/Plugins/Experimental/NeuralNetworkInference/Source/NeuralNetworkInferenceQA/Private/OperatorUnitTester.h
+++ b/Engine/Plugins/Experimental/NeuralNetworkInference/Source/NeuralNetworkInferenceQA/Private/OperatorUnitTester.h
@@ -8,7 +8,7 @@
 class FOperatorUnitTester
 {
 public:
-	static void GlobalTest(const FString& InProjectContentDir, const FString& InUnitTestRelativeDirectory);
+	static bool GlobalTest(const FString& InProjectContentDir, const FString& InUnitTestRelativeDirectory);
 
 private:
 	static void CreateRandomArraysOfFixedSeedAndDefinedSizes(TMap<int32, TArray<float>>& InOutInputTensorData, TMap<int32, TArray<float>>& InOutInputTensorDataSmall,
diff --git a/Engine/Plugins/Experimental/NeuralNetworkInference/Source/NeuralNetworkInferenceQA/Private/UnitTester.cpp b/Engine/Plugins/Experimental/NeuralNetworkInference/Source/NeuralNetworkInferenceQA/Private/UnitTester.cpp
--- a/Engine/Plugins/Experimental/NeuralNetworkInference/Source/NeuralNetworkInferenceQA/Private/UnitTester.cpp
+++ b/Engine/Plugins/Experimental/NeuralNetworkInference/Source/NeuralNetworkInferenceQA/Private/UnitTester.cpp
@@ -17,7 +17,7 @@
 /* FUnitTester static public functions
  *****************************************************************************/
 
-void FUnitTester::GlobalTest(const FString& InProjectContentDir, const FString& InModelZooRelativeDirectory, const FString& InUnitTestRelativeDirectory)
+bool FUnitTester::GlobalTest(const FString& InProjectContentDir, const FString& InModelZooRelativeDirectory, const FString& InUnitTestRelativeDirectory)
 {
 	UE_LOG(LogNeuralNetworkInferenceQA, Display, TEXT("---------------------------------------------------------------------------------------------------------------------------------"));
 	UE_LOG(LogNeuralNetworkInferenceQA, Display, TEXT("---------------------------------------------------------------------------------------------------------------------------------"));
@@ -25,19 +25,19 @@ void FUnitTester::GlobalTest(const FString& InProjectContentDir, const FString&
 	UE_LOG(LogNeuralNetworkInferenceQA, Display, TEXT("----- Starting UnitTesting() ----------------------------------------------------------------------------------------------------"));
 
 	UE_LOG(LogNeuralNetworkInferenceQA, Display, TEXT("---------- 1. Model Unit Testing (Legacy)"));
-	FLegacyModelUnitTester::GlobalTest(InProjectContentDir, InModelZooRelativeDirectory);
+	bool bDidGlobalTestPassed = FLegacyModelUnitTester::GlobalTest(InProjectContentDir, InModelZooRelativeDirectory);
 
 	UE_LOG(LogNeuralNetworkInferenceQA, Display, TEXT("---------------------------------------------------------------------------------------------------------------------------------"));
 	UE_LOG(LogNeuralNetworkInferenceQA, Display, TEXT("---------------------------------------------------------------------------------------------------------------------------------"));
 	UE_LOG(LogNeuralNetworkInferenceQA, Display, TEXT("---------------------------------------------------------------------------------------------------------------------------------"));
 	UE_LOG(LogNeuralNetworkInferenceQA, Display, TEXT("---------- 2. Model Unit Testing"));
-	FModelUnitTester::GlobalTest(InProjectContentDir, InModelZooRelativeDirectory);
+	bDidGlobalTestPassed &= FModelUnitTester::GlobalTest(InProjectContentDir, InModelZooRelativeDirectory);
 
 	UE_LOG(LogNeuralNetworkInferenceQA, Display, TEXT("---------------------------------------------------------------------------------------------------------------------------------"));
 	UE_LOG(LogNeuralNetworkInferenceQA, Display, TEXT("---------------------------------------------------------------------------------------------------------------------------------"));
 	UE_LOG(LogNeuralNetworkInferenceQA, Display, TEXT("---------------------------------------------------------------------------------------------------------------------------------"));
 	UE_LOG(LogNeuralNetworkInferenceQA, Display, TEXT("---------- 3. Operator Unit Testing"));
-	FOperatorUnitTester::GlobalTest(InProjectContentDir, InUnitTestRelativeDirectory);
+	bDidGlobalTestPassed &= FOperatorUnitTester::GlobalTest(InProjectContentDir, InUnitTestRelativeDirectory);
 
 // #if WITH_EDITOR
 // #ifdef PLATFORM_WIN64
@@ -45,9 +45,18 @@ void FUnitTester::GlobalTest(const FString& InProjectContentDir, const FString&
 // 	UE_LOG(LogNeuralNetworkInferenceQA, Display, TEXT("---------------------------------------------------------------------------------------------------------------------------------"));
 // 	UE_LOG(LogNeuralNetworkInferenceQA, Display, TEXT("---------------------------------------------------------------------------------------------------------------------------------"));
 // 	UE_LOG(LogNeuralNetworkInferenceQA, Display, TEXT("---------- 4. ONNX Runtime DLL Unit Testing (Deprecated)"));
-// 	FONNXRuntimeDLLTester::GlobalTest(InProjectContentDir, InModelZooRelativeDirectory);
+// 	bDidGlobalTestPassed &= FONNXRuntimeDLLTester::GlobalTest(InProjectContentDir, InModelZooRelativeDirectory);
 // #endif //PLATFORM_WIN64
 // #endif //WITH_EDITOR
 
-	UE_LOG(LogNeuralNetworkInferenceQA, Display, TEXT("----- UnitTesting() completed! --------------------------------------------------------------------------------------------------"));
+	if (bDidGlobalTestPassed)
+	{
+		UE_LOG(LogNeuralNetworkInferenceQA, Display, TEXT("----- UnitTesting() completed! --------------------------------------------------------------------------------------------------"));
+	}
+	else
+	{
+		UE_LOG(LogNeuralNetworkInferenceQA, Warning, TEXT("----- UnitTesting() finished with warnings/errors! --------------------------------------------------------------------------------------------------"));
+	}
+
+	return bDidGlobalTestPassed;
 }
diff --git a/Engine/Plugins/Experimental/NeuralNetworkInference/Source/NeuralNetworkInferenceQA/Private/UnitTester.h b/Engine/Plugins/Experimental/NeuralNetworkInference/Source/NeuralNetworkInferenceQA/Private/UnitTester.h
--- a/Engine/Plugins/Experimental/NeuralNetworkInference/Source/NeuralNetworkInferenceQA/Private/UnitTester.h
+++ b/Engine/Plugins/Experimental/NeuralNetworkInference/Source/NeuralNetworkInferenceQA/Private/UnitTester.h
@@ -10,5 +10,5 @@ public:
 	/**
 	 * Main unit tester function.
 	 */
-	static void GlobalTest(const FString& InProjectContentDir, const FString& InModelZooRelativeDirectory, const FString& InUnitTestRelativeDirectory);
+	static bool GlobalTest(const FString& InProjectContentDir, const FString& InModelZooRelativeDirectory, const FString& InUnitTestRelativeDirectory);
 };
diff --git a/Engine/Plugins/Experimental/NeuralNetworkInference/Source/NeuralNetworkInferenceQA/Public/NeuralNetworkInferenceQA.h b/Engine/Plugins/Experimental/NeuralNetworkInference/Source/NeuralNetworkInferenceQA/Public/NeuralNetworkInferenceQA.h
--- a/Engine/Plugins/Experimental/NeuralNetworkInference/Source/NeuralNetworkInferenceQA/Public/NeuralNetworkInferenceQA.h
+++ b/Engine/Plugins/Experimental/NeuralNetworkInference/Source/NeuralNetworkInferenceQA/Public/NeuralNetworkInferenceQA.h
@@ -16,5 +16,5 @@ class NEURALNETWORKINFERENCEQA_API UNeuralNetworkInferenceQA : public UObject
 
 public:
 	UFUNCTION(BlueprintCallable, Category = "NeuralNetworkInferenceQA")
-	static void UnitTesting();
+	static bool UnitTesting();
 };
