NeuralNetworkInference (NNI) plugin:
- Minor Cleanup.

#rnx

diff --git a/Engine/Plugins/Experimental/NeuralNetworkInference/Source/NeuralNetworkInference/Private/NeuralNetworkImplBackEndUEOnly.cpp b/Engine/Plugins/Experimental/NeuralNetworkInference/Source/NeuralNetworkInference/Private/NeuralNetworkImplBackEndUEOnly.cpp
--- a/Engine/Plugins/Experimental/NeuralNetworkInference/Source/NeuralNetworkInference/Private/NeuralNetworkImplBackEndUEOnly.cpp
+++ b/Engine/Plugins/Experimental/NeuralNetworkInference/Source/NeuralNetworkInference/Private/NeuralNetworkImplBackEndUEOnly.cpp
@@ -40,7 +40,7 @@ bool UNeuralNetwork::FImplBackEndUEOnly::Load(TSharedPtr<FImplBackEndUEOnly>& In
 		UE_LOG(LogNeuralNetworkInference, Warning, TEXT("UNeuralNetwork::FImplBackEndUEOnly::Load(): TensorManager could not be loaded from %s."), *InModelFilePath);
 		return false;
 	}
-	// Everything worked out, return true
+	// Load successful
 	return true;
 
 #else //WITH_EDITOR
