NeuralNetworkInference (NNI) plugin:
- Removed static warning for NNI.
#rnx

diff --git a/Engine/Plugins/Experimental/NeuralNetworkInference/Source/NeuralNetworkInference/Private/NeuralNetworkImplBackEndUEAndORT.cpp b/Engine/Plugins/Experimental/NeuralNetworkInference/Source/NeuralNetworkInference/Private/NeuralNetworkImplBackEndUEAndORT.cpp
--- a/Engine/Plugins/Experimental/NeuralNetworkInference/Source/NeuralNetworkInference/Private/NeuralNetworkImplBackEndUEAndORT.cpp
+++ b/Engine/Plugins/Experimental/NeuralNetworkInference/Source/NeuralNetworkInference/Private/NeuralNetworkImplBackEndUEAndORT.cpp
@@ -2,6 +2,22 @@
 
 #include "NeuralNetworkImplBackEndUEAndORT.h"
 #include "NeuralNetworkInferenceUtils.h"
+#include "RedirectCoutAndCerrToUeLog.h"
+
+//#define WITH_NNI_CPU_NOT_RECOMMENDED // Only for debugging purposes
+
+NNI_THIRD_PARTY_INCLUDES_START
+#undef check
+#undef TEXT
+#ifdef WITH_UE_AND_ORT_SUPPORT
+	#ifdef PLATFORM_WIN64
+	#include "onnxruntime/core/providers/dml/dml_provider_factory.h"
+	#endif
+	#ifdef WITH_NNI_CPU_NOT_RECOMMENDED
+	#include "onnxruntime/core/providers/nni_cpu/nni_cpu_provider_factory.h"
+	#endif //WITH_NNI_CPU_NOT_RECOMMENDED
+#endif //WITH_UE_AND_ORT_SUPPORT
+NNI_THIRD_PARTY_INCLUDES_END
 
 
 
diff --git a/Engine/Plugins/Experimental/NeuralNetworkInference/Source/NeuralNetworkInference/Private/NeuralNetworkImplBackEndUEAndORT.h b/Engine/Plugins/Experimental/NeuralNetworkInference/Source/NeuralNetworkInference/Private/NeuralNetworkImplBackEndUEAndORT.h
--- a/Engine/Plugins/Experimental/NeuralNetworkInference/Source/NeuralNetworkInference/Private/NeuralNetworkImplBackEndUEAndORT.h
+++ b/Engine/Plugins/Experimental/NeuralNetworkInference/Source/NeuralNetworkInference/Private/NeuralNetworkImplBackEndUEAndORT.h
@@ -4,22 +4,12 @@
 
 #include "NeuralNetwork.h"
 
-//#define WITH_NNI_CPU_NOT_RECOMMENDED // Only for debugging purposes
-
 #include "ThirdPartyWarningDisabler.h"
 NNI_THIRD_PARTY_INCLUDES_START
 #undef check
 #undef TEXT
 #ifdef WITH_UE_AND_ORT_SUPPORT
-	#include "RedirectCoutAndCerrToUeLog.h"
-
 	#include "onnxruntime/core/session/onnxruntime_cxx_api.h"
-	#ifdef PLATFORM_WIN64
-	#include "onnxruntime/core/providers/dml/dml_provider_factory.h"
-	#endif
-	#ifdef WITH_NNI_CPU_NOT_RECOMMENDED
-	#include "onnxruntime/core/providers/nni_cpu/nni_cpu_provider_factory.h"
-	#endif //WITH_NNI_CPU_NOT_RECOMMENDED
 #endif //WITH_UE_AND_ORT_SUPPORT
 NNI_THIRD_PARTY_INCLUDES_END
 
