NeuralNetworkInference (NNI) plugin:
- Use UE D3D12 device and command queue when creating DirectML execution provider

#rnx

diff --git a/Engine/Plugins/Experimental/NeuralNetworkInference/Source/NeuralNetworkInference/NeuralNetworkInference.Build.cs b/Engine/Plugins/Experimental/NeuralNetworkInference/Source/NeuralNetworkInference/NeuralNetworkInference.Build.cs
--- a/Engine/Plugins/Experimental/NeuralNetworkInference/Source/NeuralNetworkInference/NeuralNetworkInference.Build.cs
+++ b/Engine/Plugins/Experimental/NeuralNetworkInference/Source/NeuralNetworkInference/NeuralNetworkInference.Build.cs
@@ -1,6 +1,7 @@
 // Copyright Epic Games, Inc. All Rights Reserved.
 
 using UnrealBuildTool;
+using System.IO;
 
 public class NeuralNetworkInference : ModuleRules
 {
@@ -41,6 +42,7 @@ public class NeuralNetworkInference : ModuleRules
 				"NeuralNetworkInferenceShaders",
 				"RenderCore",
 				"RHI",
+				"RHICore",
 				// ORT-related
 				"Projects",
 				"ThirdPartyHelperAndDLLLoader"
@@ -60,6 +62,26 @@ public class NeuralNetworkInference : ModuleRules
 
 		if (bIsORTSupported)
 		{
+			if (Target.Platform == UnrealTargetPlatform.Win64)
+			{
+				// Borrowed from Plugins\Media\ImgMedia (we need to use FD3D12DynamicRHI
+				PrivateIncludePaths.AddRange(
+					new string[]{
+						//required for "D3D12RHIPrivate.h"
+						Path.Combine(EngineDirectory, "Source/Runtime/D3D12RHI/Private"),
+						Path.Combine(EngineDirectory, "Source/Runtime/D3D12RHI/Private/Windows")
+					});
+
+				PrivateDependencyModuleNames.AddRange
+					(
+					new string[] {
+						"D3D12RHI",
+						"DirectML"
+					}
+				);
+
+				AddEngineThirdPartyPrivateStaticDependencies(Target, "DX12");
+			}
 
 			PrivateDependencyModuleNames.AddRange
 				(
diff --git a/Engine/Plugins/Experimental/NeuralNetworkInference/Source/NeuralNetworkInference/Private/NeuralNetworkImplBackEndUEAndORT.cpp b/Engine/Plugins/Experimental/NeuralNetworkInference/Source/NeuralNetworkInference/Private/NeuralNetworkImplBackEndUEAndORT.cpp
--- a/Engine/Plugins/Experimental/NeuralNetworkInference/Source/NeuralNetworkInference/Private/NeuralNetworkImplBackEndUEAndORT.cpp
+++ b/Engine/Plugins/Experimental/NeuralNetworkInference/Source/NeuralNetworkInference/Private/NeuralNetworkImplBackEndUEAndORT.cpp
@@ -4,6 +4,21 @@
 #include "NeuralNetworkInferenceUtils.h"
 #include "RedirectCoutAndCerrToUeLog.h"
 
+#if defined(WITH_UE_AND_ORT_SUPPORT) && defined(PLATFORM_WIN64)
+	#include "RHI.h"
+	#include "DynamicRHI.h"
+	
+	// Disable NOMINMAX & WIN32_LEAN_AND_MEAN defines to avoid compiler warnings
+	#pragma push_macro("NOMINMAX")
+	#pragma push_macro("WIN32_LEAN_AND_MEAN")
+	#undef NOMINMAX
+	#undef WIN32_LEAN_AND_MEAN
+	#include "D3D12RHIPrivate.h"
+	#pragma pop_macro("WIN32_LEAN_AND_MEAN")
+	#pragma pop_macro("NOMINMAX")
+
+#endif
+
 //#define WITH_NNI_CPU_NOT_RECOMMENDED // Only for debugging purposes
 
 NNI_THIRD_PARTY_INCLUDES_START
@@ -19,6 +34,66 @@ NNI_THIRD_PARTY_INCLUDES_START
 #endif //WITH_UE_AND_ORT_SUPPORT
 NNI_THIRD_PARTY_INCLUDES_END
 
+#if defined(WITH_UE_AND_ORT_SUPPORT) && defined(PLATFORM_WIN64)
+
+/**
+* Helper class that maintains a list of created DML Devices for given ID3D12Device
+*/
+class DMLDeviceList
+{
+public:
+
+	IDMLDevice* GetDMLDevice(ID3D12Device* Device)
+	{
+		for (size_t c = 0; c < Entries.Num(); ++c)
+		{
+			if (Entries[c].Device == Device)
+			{
+				return Entries[c].DmlDevice;
+			}
+		}
+
+		return Add(Device);
+	}
+
+private:
+
+	IDMLDevice* Add(ID3D12Device* Device)
+	{
+		// Create new DML Device
+		IDMLDevice* DmlDevice = nullptr;
+
+		DML_CREATE_DEVICE_FLAGS DmlCreateFlags = DML_CREATE_DEVICE_FLAG_NONE;
+
+		DMLCreateDevice1(Device,
+			DmlCreateFlags,
+			DML_FEATURE_LEVEL_2_0,
+			IID_PPV_ARGS(&DmlDevice));
+
+		if (!DmlDevice)
+		{
+			UE_LOG(LogNeuralNetworkInference, Error, TEXT("Failed to create DML device."));
+			return nullptr;
+		}
+
+		Entries.Push(DMLDeviceEntry{ Device, DmlDevice });
+
+		return DmlDevice;
+	}
+
+	struct DMLDeviceEntry
+	{
+		ID3D12Device* Device;
+		IDMLDevice* DmlDevice;
+	};
+
+	TArray< DMLDeviceEntry >		Entries;
+};
+
+static DMLDeviceList GDMLDevices;
+
+#endif
+
 
 
 /* UNeuralNetwork public functions
@@ -233,9 +308,43 @@ bool UNeuralNetwork::FImplBackEndUEAndORT::ConfigureMembers(const ENeuralDeviceT
 	if (InDeviceType == ENeuralDeviceType::GPU)
 	{
 #ifdef PLATFORM_WIN64
+		// To create a DirectML device we need to check that we're using DX12 first
+		FString RhiName = GDynamicRHI->GetName();
+
+		if (RhiName != TEXT("D3D12"))
+		{
+			UE_LOG(LogNeuralNetworkInference, Warning, TEXT("Only D3D12 rendering is working with NNI"));
+			return false;
+		}
+
+		// Get adapter's D3D12 device that we would like to share with DirectML execution provider
+		// NOTE: For now we're only using first device that has Dadapter 0 and device 0
+		FD3D12DynamicRHI* Rhi = static_cast<FD3D12DynamicRHI*>(GDynamicRHI);
+
+		if (Rhi->GetNumAdapters() > 1 || Rhi->GetAdapter(0).GetDesc().NumDeviceNodes > 1)
+		{
+			UE_LOG(LogNeuralNetworkInference, Warning, TEXT("There are multiple adapters and/or multiple devices, using device at index 0"));
+			return false;
+		}
+
+		ID3D12Device* NativeDevice = Rhi->GetAdapter(0).GetD3DDevice();
+
+		// Make sure that we have one DMLDevice per D3D12 device
+		IDMLDevice* DmlDevice = GDMLDevices.GetDMLDevice(NativeDevice);
+
+		if (!DmlDevice)
+		{
+			UE_LOG(LogNeuralNetworkInference, Warning, TEXT("Got invalid DML device"));
+			return false;
+		}
+
+		// Get a ID3D12CommandQueue as well
+		// TODO: Should we create our own queue?
+		ID3D12CommandQueue* NativeCmdQ = Rhi->RHIGetD3DCommandQueue();
+
 		// ORT GPU (Direct ML)
 		SessionOptions->SetGraphOptimizationLevel(GraphOptimizationLevel::ORT_ENABLE_ALL); // ORT_ENABLE_ALL, ORT_ENABLE_EXTENDED, ORT_ENABLE_BASIC, ORT_DISABLE_ALL
-		if (OrtSessionOptionsAppendExecutionProvider_DML(*SessionOptions, 0))
+		if (OrtSessionOptionsAppendExecutionProviderEx_DML(*SessionOptions, DmlDevice, NativeCmdQ))
 		{
 			UE_LOG(LogNeuralNetworkInference, Warning, TEXT("Some error occurred."));
 			return false;
