NeuralNetworkInference (NNI) plugin:
- Fixed crashes in NNI for DX11 (default) UE.

#rnx

diff --git a/Engine/Plugins/Experimental/NeuralNetworkInference/Source/NeuralNetworkInference/Private/NeuralNetworkImplBackEndUEAndORT.cpp b/Engine/Plugins/Experimental/NeuralNetworkInference/Source/NeuralNetworkInference/Private/NeuralNetworkImplBackEndUEAndORT.cpp
--- a/Engine/Plugins/Experimental/NeuralNetworkInference/Source/NeuralNetworkInference/Private/NeuralNetworkImplBackEndUEAndORT.cpp
+++ b/Engine/Plugins/Experimental/NeuralNetworkInference/Source/NeuralNetworkInference/Private/NeuralNetworkImplBackEndUEAndORT.cpp
@@ -309,12 +309,21 @@ bool UNeuralNetwork::FImplBackEndUEAndORT::ConfigureMembers(const ENeuralDeviceT
 	{
 #ifdef PLATFORM_WIN64
 		// To create a DirectML device we need to check that we're using DX12 first
-		FString RhiName = GDynamicRHI->GetName();
+		const FString RHIName = GDynamicRHI->GetName();
 
-		if (RhiName != TEXT("D3D12"))
+		// @todo for Nebojša.Dragosavac & Ginés.Hidalgo: Handle the default DX11 nicely
+		if (RHIName != TEXT("D3D12"))
 		{
-			UE_LOG(LogNeuralNetworkInference, Warning, TEXT("Only D3D12 rendering is working with NNI"));
-			return false;
+			//UE_LOG(LogNeuralNetworkInference, Warning, TEXT("UNeuralNetwork::ConfigureMembers(): Only DX12 (\"D3D12\")) rendering is compatible with UEAndORT-based NNI. Instead, \"%s\" was used."), *RHIName);
+			//return false;
+
+			// Temporary workaround to avoid breaking NNI
+			if (OrtSessionOptionsAppendExecutionProvider_DML(*SessionOptions, 0))
+			{
+				UE_LOG(LogNeuralNetworkInference, Warning, TEXT("UNeuralNetwork::ConfigureMembers(): Some error occurred when using OrtSessionOptionsAppendExecutionProvider_DML()."));
+				return false;
+			}
+			return true;
 		}
 
 		// Get adapter's D3D12 device that we would like to share with DirectML execution provider
@@ -323,7 +332,8 @@ bool UNeuralNetwork::FImplBackEndUEAndORT::ConfigureMembers(const ENeuralDeviceT
 
 		if (Rhi->GetNumAdapters() > 1 || Rhi->GetAdapter(0).GetDesc().NumDeviceNodes > 1)
 		{
-			UE_LOG(LogNeuralNetworkInference, Warning, TEXT("There are multiple adapters and/or multiple devices, using device at index 0"));
+			UE_LOG(LogNeuralNetworkInference, Warning, TEXT("UNeuralNetwork::ConfigureMembers(): There are multiple (%d) adapters and/or multiple (%d) devices, using device at index 0."),
+				Rhi->GetNumAdapters(), Rhi->GetAdapter(0).GetDesc().NumDeviceNodes);
 			return false;
 		}
 
@@ -334,7 +344,7 @@ bool UNeuralNetwork::FImplBackEndUEAndORT::ConfigureMembers(const ENeuralDeviceT
 
 		if (!DmlDevice)
 		{
-			UE_LOG(LogNeuralNetworkInference, Warning, TEXT("Got invalid DML device"));
+			UE_LOG(LogNeuralNetworkInference, Warning, TEXT("UNeuralNetwork::ConfigureMembers(): Invalid DML device found."));
 			return false;
 		}
 
@@ -346,7 +356,7 @@ bool UNeuralNetwork::FImplBackEndUEAndORT::ConfigureMembers(const ENeuralDeviceT
 		SessionOptions->SetGraphOptimizationLevel(GraphOptimizationLevel::ORT_ENABLE_ALL); // ORT_ENABLE_ALL, ORT_ENABLE_EXTENDED, ORT_ENABLE_BASIC, ORT_DISABLE_ALL
 		if (OrtSessionOptionsAppendExecutionProviderEx_DML(*SessionOptions, DmlDevice, NativeCmdQ))
 		{
-			UE_LOG(LogNeuralNetworkInference, Warning, TEXT("Some error occurred."));
+			UE_LOG(LogNeuralNetworkInference, Warning, TEXT("UNeuralNetwork::ConfigureMembers(): Some error occurred when using OrtSessionOptionsAppendExecutionProviderEx_DML()."));
 			return false;
 		}
 		return true; // @todo: Remove this line when NNI_HLSL is working
