NeuralNetworkInference (NNI) plugin:
- DMLDeviceList could be a "common name", Epic coding standard is to rename it as something that nobody else on the whole UE could accidentally re-define, renamed as  FPrivateNeuralNetwork::DMLDeviceList.
- FPrivateNeuralNetwork_DMLDeviceList split into signature and body to follow NNI coding standard.

#rnx

diff --git a/Engine/Plugins/Experimental/NeuralNetworkInference/Source/NeuralNetworkInference/Private/NeuralNetworkImplBackEndUEAndORT.cpp b/Engine/Plugins/Experimental/NeuralNetworkInference/Source/NeuralNetworkInference/Private/NeuralNetworkImplBackEndUEAndORT.cpp
--- a/Engine/Plugins/Experimental/NeuralNetworkInference/Source/NeuralNetworkInference/Private/NeuralNetworkImplBackEndUEAndORT.cpp
+++ b/Engine/Plugins/Experimental/NeuralNetworkInference/Source/NeuralNetworkInference/Private/NeuralNetworkImplBackEndUEAndORT.cpp
@@ -36,61 +36,73 @@ NNI_THIRD_PARTY_INCLUDES_END
 
 #if defined(WITH_UE_AND_ORT_SUPPORT) && defined(PLATFORM_WIN64)
 
-/**
-* Helper class that maintains a list of created DML Devices for given ID3D12Device
-*/
-class DMLDeviceList
+
+
+/* DMLDeviceList auxiliary class
+ *****************************************************************************/
+class FPrivateNeuralNetwork
 {
 public:
-
-	IDMLDevice* GetDMLDevice(ID3D12Device* Device)
+	/**
+	 * Helper class that maintains a list of created DML Devices for given ID3D12Device
+	 */
+	class DMLDeviceList
 	{
-		for (size_t c = 0; c < Entries.Num(); ++c)
+	public:
+		IDMLDevice* GetDMLDevice(ID3D12Device* Device);
+
+	private:
+		IDMLDevice* Add(ID3D12Device* Device);
+
+		struct DMLDeviceEntry
 		{
-			if (Entries[c].Device == Device)
-			{
-				return Entries[c].DmlDevice;
-			}
-		}
+			ID3D12Device* Device;
+			IDMLDevice* DmlDevice;
+		};
 
-		return Add(Device);
-	}
+		TArray<DMLDeviceEntry>		Entries;
+	};
 
-private:
+	static DMLDeviceList GDMLDevices;
+};
 
-	IDMLDevice* Add(ID3D12Device* Device)
+IDMLDevice* FPrivateNeuralNetwork::DMLDeviceList::GetDMLDevice(ID3D12Device* Device)
+{
+	for (size_t c = 0; c < Entries.Num(); ++c)
 	{
-		// Create new DML Device
-		IDMLDevice* DmlDevice = nullptr;
+		if (Entries[c].Device == Device)
+		{
+			return Entries[c].DmlDevice;
+		}
+	}
 
-		DML_CREATE_DEVICE_FLAGS DmlCreateFlags = DML_CREATE_DEVICE_FLAG_NONE;
+	return Add(Device);
+}
 
-		DMLCreateDevice1(Device,
-			DmlCreateFlags,
-			DML_FEATURE_LEVEL_2_0,
-			IID_PPV_ARGS(&DmlDevice));
+IDMLDevice* FPrivateNeuralNetwork::DMLDeviceList::Add(ID3D12Device* Device)
+{
+	// Create new DML Device
+	IDMLDevice* DmlDevice = nullptr;
 
-		if (!DmlDevice)
-		{
-			UE_LOG(LogNeuralNetworkInference, Error, TEXT("Failed to create DML device."));
-			return nullptr;
-		}
+	DML_CREATE_DEVICE_FLAGS DmlCreateFlags = DML_CREATE_DEVICE_FLAG_NONE;
 
-		Entries.Push(DMLDeviceEntry{ Device, DmlDevice });
+	DMLCreateDevice1(Device,
+		DmlCreateFlags,
+		DML_FEATURE_LEVEL_2_0,
+		IID_PPV_ARGS(&DmlDevice));
 
-		return DmlDevice;
+	if (!DmlDevice)
+	{
+		UE_LOG(LogNeuralNetworkInference, Error, TEXT("Failed to create DML device."));
+		return nullptr;
 	}
 
-	struct DMLDeviceEntry
-	{
-		ID3D12Device* Device;
-		IDMLDevice* DmlDevice;
-	};
+	Entries.Push(DMLDeviceEntry{ Device, DmlDevice });
 
-	TArray< DMLDeviceEntry >		Entries;
-};
+	return DmlDevice;
+}
 
-static DMLDeviceList GDMLDevices;
+FPrivateNeuralNetwork::DMLDeviceList FPrivateNeuralNetwork::GDMLDevices;
 
 #endif
 
@@ -340,7 +352,7 @@ bool UNeuralNetwork::FImplBackEndUEAndORT::ConfigureMembers(const ENeuralDeviceT
 		ID3D12Device* NativeDevice = Rhi->GetAdapter(0).GetD3DDevice();
 
 		// Make sure that we have one DMLDevice per D3D12 device
-		IDMLDevice* DmlDevice = GDMLDevices.GetDMLDevice(NativeDevice);
+		IDMLDevice* DmlDevice = FPrivateNeuralNetwork::GDMLDevices.GetDMLDevice(NativeDevice);
 
 		if (!DmlDevice)
 		{
