NeuralNetworkInference (NNI) plugin:
- DMLDeviceList only uses memory it NNI-GPU is used at least once.

#rnx

diff --git a/Engine/Plugins/Experimental/NeuralNetworkInference/Source/NeuralNetworkInference/Private/NeuralNetworkImplBackEndUEAndORT.cpp b/Engine/Plugins/Experimental/NeuralNetworkInference/Source/NeuralNetworkInference/Private/NeuralNetworkImplBackEndUEAndORT.cpp
--- a/Engine/Plugins/Experimental/NeuralNetworkInference/Source/NeuralNetworkInference/Private/NeuralNetworkImplBackEndUEAndORT.cpp
+++ b/Engine/Plugins/Experimental/NeuralNetworkInference/Source/NeuralNetworkInference/Private/NeuralNetworkImplBackEndUEAndORT.cpp
@@ -38,15 +38,18 @@ NNI_THIRD_PARTY_INCLUDES_END
 
 
 
-/* DMLDeviceList auxiliary class
+/* FPrivateNeuralNetwork auxiliary class
  *****************************************************************************/
 class FPrivateNeuralNetwork
 {
 public:
+	static IDMLDevice* GetDMLDevice(ID3D12Device* Device);
+
+private:
 	/**
 	 * Helper class that maintains a list of created DML Devices for given ID3D12Device
 	 */
-	class DMLDeviceList
+	class FDMLDeviceList
 	{
 	public:
 		IDMLDevice* GetDMLDevice(ID3D12Device* Device);
@@ -62,11 +65,15 @@ public:
 
 		TArray<DMLDeviceEntry>		Entries;
 	};
-
-	static DMLDeviceList GDMLDevices;
 };
 
-IDMLDevice* FPrivateNeuralNetwork::DMLDeviceList::GetDMLDevice(ID3D12Device* Device)
+IDMLDevice* FPrivateNeuralNetwork::GetDMLDevice(ID3D12Device* Device)
+{
+	static FDMLDeviceList DMLDeviceList;
+	return DMLDeviceList.GetDMLDevice(Device);
+}
+
+IDMLDevice* FPrivateNeuralNetwork::FDMLDeviceList::GetDMLDevice(ID3D12Device* Device)
 {
 	for (size_t c = 0; c < Entries.Num(); ++c)
 	{
@@ -79,7 +86,7 @@ IDMLDevice* FPrivateNeuralNetwork::DMLDeviceList::GetDMLDevice(ID3D12Device* Dev
 	return Add(Device);
 }
 
-IDMLDevice* FPrivateNeuralNetwork::DMLDeviceList::Add(ID3D12Device* Device)
+IDMLDevice* FPrivateNeuralNetwork::FDMLDeviceList::Add(ID3D12Device* Device)
 {
 	// Create new DML Device
 	IDMLDevice* DmlDevice = nullptr;
@@ -102,8 +109,6 @@ IDMLDevice* FPrivateNeuralNetwork::DMLDeviceList::Add(ID3D12Device* Device)
 	return DmlDevice;
 }
 
-FPrivateNeuralNetwork::DMLDeviceList FPrivateNeuralNetwork::GDMLDevices;
-
 #endif
 
 
@@ -352,7 +357,7 @@ bool UNeuralNetwork::FImplBackEndUEAndORT::ConfigureMembers(const ENeuralDeviceT
 		ID3D12Device* NativeDevice = Rhi->GetAdapter(0).GetD3DDevice();
 
 		// Make sure that we have one DMLDevice per D3D12 device
-		IDMLDevice* DmlDevice = FPrivateNeuralNetwork::GDMLDevices.GetDMLDevice(NativeDevice);
+		IDMLDevice* DmlDevice = FPrivateNeuralNetwork::GetDMLDevice(NativeDevice);
 
 		if (!DmlDevice)
 		{
