NeuralNetworkInference (NNI) plugin:
- Fixed CIS error: FPrivateNeuralNetwork renamed as FPrivateImplBackEndUEAndORT.

#rnx

diff --git a/Engine/Plugins/Experimental/NeuralNetworkInference/Source/NeuralNetworkInference/Private/NeuralNetworkImplBackEndUEAndORT.cpp b/Engine/Plugins/Experimental/NeuralNetworkInference/Source/NeuralNetworkInference/Private/NeuralNetworkImplBackEndUEAndORT.cpp
--- a/Engine/Plugins/Experimental/NeuralNetworkInference/Source/NeuralNetworkInference/Private/NeuralNetworkImplBackEndUEAndORT.cpp
+++ b/Engine/Plugins/Experimental/NeuralNetworkInference/Source/NeuralNetworkInference/Private/NeuralNetworkImplBackEndUEAndORT.cpp
@@ -38,9 +38,9 @@ NNI_THIRD_PARTY_INCLUDES_END
 
 
 
-/* FPrivateNeuralNetwork auxiliary class
+/* FPrivateImplBackEndUEAndORT auxiliary class
  *****************************************************************************/
-class FPrivateNeuralNetwork
+class FPrivateImplBackEndUEAndORT
 {
 public:
 	static IDMLDevice* GetDMLDevice(ID3D12Device* Device);
@@ -67,13 +67,13 @@ private:
 	};
 };
 
-IDMLDevice* FPrivateNeuralNetwork::GetDMLDevice(ID3D12Device* Device)
+IDMLDevice* FPrivateImplBackEndUEAndORT::GetDMLDevice(ID3D12Device* Device)
 {
 	static FDMLDeviceList DMLDeviceList;
 	return DMLDeviceList.GetDMLDevice(Device);
 }
 
-IDMLDevice* FPrivateNeuralNetwork::FDMLDeviceList::GetDMLDevice(ID3D12Device* Device)
+IDMLDevice* FPrivateImplBackEndUEAndORT::FDMLDeviceList::GetDMLDevice(ID3D12Device* Device)
 {
 	for (size_t c = 0; c < Entries.Num(); ++c)
 	{
@@ -86,7 +86,7 @@ IDMLDevice* FPrivateNeuralNetwork::FDMLDeviceList::GetDMLDevice(ID3D12Device* De
 	return Add(Device);
 }
 
-IDMLDevice* FPrivateNeuralNetwork::FDMLDeviceList::Add(ID3D12Device* Device)
+IDMLDevice* FPrivateImplBackEndUEAndORT::FDMLDeviceList::Add(ID3D12Device* Device)
 {
 	// Create new DML Device
 	IDMLDevice* DmlDevice = nullptr;
@@ -357,7 +357,7 @@ bool UNeuralNetwork::FImplBackEndUEAndORT::ConfigureMembers(const ENeuralDeviceT
 		ID3D12Device* NativeDevice = Rhi->GetAdapter(0).GetD3DDevice();
 
 		// Make sure that we have one DMLDevice per D3D12 device
-		IDMLDevice* DmlDevice = FPrivateNeuralNetwork::GetDMLDevice(NativeDevice);
+		IDMLDevice* DmlDevice = FPrivateImplBackEndUEAndORT::GetDMLDevice(NativeDevice);
 
 		if (!DmlDevice)
 		{
