NeuralNetworkInference (NNI) plugin:
- DMLDeviceList made thread safe.

#rnx

diff --git a/Engine/Plugins/Experimental/NeuralNetworkInference/Source/NeuralNetworkInference/Private/NeuralNetworkImplBackEndUEAndORT.cpp b/Engine/Plugins/Experimental/NeuralNetworkInference/Source/NeuralNetworkInference/Private/NeuralNetworkImplBackEndUEAndORT.cpp
--- a/Engine/Plugins/Experimental/NeuralNetworkInference/Source/NeuralNetworkInference/Private/NeuralNetworkImplBackEndUEAndORT.cpp
+++ b/Engine/Plugins/Experimental/NeuralNetworkInference/Source/NeuralNetworkInference/Private/NeuralNetworkImplBackEndUEAndORT.cpp
@@ -5,6 +5,7 @@
 #include "RedirectCoutAndCerrToUeLog.h"
 
 #if defined(WITH_UE_AND_ORT_SUPPORT) && defined(PLATFORM_WIN64)
+	#include "HAL/CriticalSection.h"
 	#include "RHI.h"
 	#include "DynamicRHI.h"
 	
@@ -34,16 +35,16 @@ NNI_THIRD_PARTY_INCLUDES_START
 #endif //WITH_UE_AND_ORT_SUPPORT
 NNI_THIRD_PARTY_INCLUDES_END
 
-#if defined(WITH_UE_AND_ORT_SUPPORT) && defined(PLATFORM_WIN64)
 
 
+#if defined(WITH_UE_AND_ORT_SUPPORT) && defined(PLATFORM_WIN64)
 
 /* FPrivateImplBackEndUEAndORT auxiliary class
  *****************************************************************************/
 class FPrivateImplBackEndUEAndORT
 {
 public:
-	static IDMLDevice* GetDMLDevice(ID3D12Device* Device);
+	static IDMLDevice* GetDMLDeviceThreadSafe(ID3D12Device* Device);
 
 private:
 	/**
@@ -67,9 +68,11 @@ private:
 	};
 };
 
-IDMLDevice* FPrivateImplBackEndUEAndORT::GetDMLDevice(ID3D12Device* Device)
+IDMLDevice* FPrivateImplBackEndUEAndORT::GetDMLDeviceThreadSafe(ID3D12Device* Device)
 {
+	static FCriticalSection CriticalSection; /* Critical section to protect GetDMLDeviceThreadSafe from being called simultaneously from multiple threads. */
 	static FDMLDeviceList DMLDeviceList;
+	FScopeLock Lock(&CriticalSection);
 	return DMLDeviceList.GetDMLDevice(Device);
 }
 
@@ -357,7 +360,7 @@ bool UNeuralNetwork::FImplBackEndUEAndORT::ConfigureMembers(const ENeuralDeviceT
 		ID3D12Device* NativeDevice = Rhi->GetAdapter(0).GetD3DDevice();
 
 		// Make sure that we have one DMLDevice per D3D12 device
-		IDMLDevice* DmlDevice = FPrivateImplBackEndUEAndORT::GetDMLDevice(NativeDevice);
+		IDMLDevice* DmlDevice = FPrivateImplBackEndUEAndORT::GetDMLDeviceThreadSafe(NativeDevice);
 
 		if (!DmlDevice)
 		{
