NeuralNetworkInference (NNI) plugin:
- ModelProtoFileReader/ModelProtoConverter: C++/Build.cs changes to enable it in game without crashing due to the linking error when used with ORT. However, ModelConverter.cpp disabled in C++ code until linking error properly fixed.
- Removed WITH_EDITOR from multiple NNI locations, as it should now be usable in game and editor after fixing linking error.
- This CL unblocks combining NNI ORT and Gin�s' back end.

#rnx

diff --git a/Engine/Plugins/Experimental/NeuralNetworkInference/NeuralNetworkInference.uplugin b/Engine/Plugins/Experimental/NeuralNetworkInference/NeuralNetworkInference.uplugin
--- a/Engine/Plugins/Experimental/NeuralNetworkInference/NeuralNetworkInference.uplugin
+++ b/Engine/Plugins/Experimental/NeuralNetworkInference/NeuralNetworkInference.uplugin
@@ -55,12 +55,12 @@
 		},
 		{
 			"Name": "ModelProtoFileReader",
-			"Type": "Editor",
+			"Type": "RuntimeAndProgram",
 			"LoadingPhase": "Default"
 		},
 		{
 			"Name": "ModelProtoConverter",
-			"Type": "Editor",
+			"Type": "RuntimeAndProgram",
 			"LoadingPhase": "Default"
 		},
 		{
diff --git a/Engine/Plugins/Experimental/NeuralNetworkInference/Source/NeuralNetworkInference/NeuralNetworkInference.Build.cs b/Engine/Plugins/Experimental/NeuralNetworkInference/Source/NeuralNetworkInference/NeuralNetworkInference.Build.cs
--- a/Engine/Plugins/Experimental/NeuralNetworkInference/Source/NeuralNetworkInference/NeuralNetworkInference.Build.cs
+++ b/Engine/Plugins/Experimental/NeuralNetworkInference/Source/NeuralNetworkInference/NeuralNetworkInference.Build.cs
@@ -39,6 +39,7 @@ public class NeuralNetworkInference : ModuleRules
 			(
 			new string[] {
 				// Backend-related
+				"ModelProtoFileReader",
 				"NeuralNetworkInferenceShaders",
 				"RenderCore",
 				"RHI",
@@ -49,17 +50,6 @@ public class NeuralNetworkInference : ModuleRules
 			}
 		);
 
-		// Editor-only
-		if (Target.bBuildEditor)
-		{
-			PrivateDependencyModuleNames.AddRange
-				(
-				new string[] {
-					"ModelProtoFileReader"
-				}
-			);
-		}
-
 		if (bIsORTSupported)
 		{
 			PrivateDependencyModuleNames.AddRange
diff --git a/Engine/Plugins/Experimental/NeuralNetworkInference/Source/NeuralNetworkInference/Private/GraphProtoToNeuralNetworkConverter.cpp b/Engine/Plugins/Experimental/NeuralNetworkInference/Source/NeuralNetworkInference/Private/GraphProtoToNeuralNetworkConverter.cpp
--- a/Engine/Plugins/Experimental/NeuralNetworkInference/Source/NeuralNetworkInference/Private/GraphProtoToNeuralNetworkConverter.cpp
+++ b/Engine/Plugins/Experimental/NeuralNetworkInference/Source/NeuralNetworkInference/Private/GraphProtoToNeuralNetworkConverter.cpp
@@ -26,7 +26,6 @@ bool FGraphProtoToNeuralNetworkConverter::Translate(TArray<TSharedPtr<FNeuralOpe
 	return bResult;
 }
 
-#if WITH_EDITOR
 bool FGraphProtoToNeuralNetworkConverter::Translate(TArray<TSharedPtr<FNeuralOperator>>& OutOperators, FNeuralTensorManager& OutTensorManager, const FGraphProto& InGraphProto,
 	const FString& InModelPath)
 {
@@ -124,7 +123,6 @@ bool FGraphProtoToNeuralNetworkConverter::Translate(TArray<TSharedPtr<FNeuralOpe
 	}
 	return bWasSuccessful;
 }
-#endif //WITH_EDITOR
 
 
 
@@ -159,7 +157,6 @@ bool FGraphProtoToNeuralNetworkConverter::CreateOperatorsAndEditTensorArray(TArr
 			// FNeuralTensor was never modified, so it must come from a weight tensor (regardless of whether it is also an absolute output of the network)
 			else
 			{
-#if WITH_EDITOR
 				const FTensorProto& TensorProto = *FModelProto::FindElementInArray(InputTensorName, InGraphProto.Initializer, /*bMustValueBeFound*/true);
 				// Is it also an absolute output of the network?
 				const int64 TensorIndex = InOutTensors.Num();
@@ -185,10 +182,6 @@ bool FGraphProtoToNeuralNetworkConverter::CreateOperatorsAndEditTensorArray(TArr
 				// Update InOutNameIndexMap and OperatorInputTensors
 				InOutNameIndexMap.Add(InputTensorName, TensorIndex);
 				OperatorInputTensors.Push(&NewTensor);
-#else
-				UE_LOG(LogNeuralNetworkInference, Warning, TEXT("FGraphProtoToNeuralNetworkConverter::CreateOperatorsAndEditTensorArray(): This code should never be called in non-Editor mode."));
-				return false;
-#endif //WITH_EDITOR
 			}
 		}
 		// Nodes that can be replaced for inlined functions and do not need to be executed
diff --git a/Engine/Plugins/Experimental/NeuralNetworkInference/Source/NeuralNetworkInference/Private/GraphProtoToNeuralNetworkConverter.h b/Engine/Plugins/Experimental/NeuralNetworkInference/Source/NeuralNetworkInference/Private/GraphProtoToNeuralNetworkConverter.h
--- a/Engine/Plugins/Experimental/NeuralNetworkInference/Source/NeuralNetworkInference/Private/GraphProtoToNeuralNetworkConverter.h
+++ b/Engine/Plugins/Experimental/NeuralNetworkInference/Source/NeuralNetworkInference/Private/GraphProtoToNeuralNetworkConverter.h
@@ -18,12 +18,10 @@ public:
 	 *
 	 * 2 alternatives:
 	 * Without InModelPath, InOutTensorManager is assumed to be a const input variable (assuming the tensors are all already loaded, e.g., on de-serialization).
-	 * With InModelPath (Editor-only function), InOutTensorManager is assumed to be an output variable that will be filled based on InModelPath.
+	 * With InModelPath, InOutTensorManager is assumed to be an output variable that will be filled based on InModelPath.
 	 */
 	static bool Translate(TArray<TSharedPtr<FNeuralOperator>>& OutOperators, FNeuralTensorManager& InTensorManager, const FGraphProto& InGraphProto);
-#if WITH_EDITOR
 	static bool Translate(TArray<TSharedPtr<FNeuralOperator>>& OutOperators, FNeuralTensorManager& OutTensorManager, const FGraphProto& InGraphProto, const FString& InModelPath);
-#endif //WITH_EDITOR
 
 private:
 	/**
diff --git a/Engine/Plugins/Experimental/NeuralNetworkInference/Source/NeuralNetworkInference/Private/NeuralNetwork.cpp b/Engine/Plugins/Experimental/NeuralNetworkInference/Source/NeuralNetworkInference/Private/NeuralNetwork.cpp
--- a/Engine/Plugins/Experimental/NeuralNetworkInference/Source/NeuralNetworkInference/Private/NeuralNetwork.cpp
+++ b/Engine/Plugins/Experimental/NeuralNetworkInference/Source/NeuralNetworkInference/Private/NeuralNetwork.cpp
@@ -131,11 +131,7 @@ bool UNeuralNetwork::Load()
 	// UEOnly
 	else if (BackEndForCurrentPlatform == ENeuralBackEnd::UEOnly)
 	{
-#if WITH_EDITOR
 		bIsLoaded = UNeuralNetwork::FImplBackEndUEOnly::Load(ImplBackEndUEOnly, /*ModelReadFromFileInBytes*/ ModelFullFilePath);
-#else //WITH_EDITOR
-		UE_LOG(LogNeuralNetworkInference, Warning, TEXT("UNeuralNetwork::Load(const FString& InModelFilePath) only supported in Editor (WITH_EDITOR) for ENeuralBackEnd::UEOnly."));
-#endif //WITH_EDITOR
 	}
 	// Unknown
 	else
diff --git a/Engine/Plugins/Experimental/NeuralNetworkInference/Source/NeuralNetworkInference/Private/NeuralNetworkImplBackEndUEAndORT.cpp b/Engine/Plugins/Experimental/NeuralNetworkInference/Source/NeuralNetworkInference/Private/NeuralNetworkImplBackEndUEAndORT.cpp
--- a/Engine/Plugins/Experimental/NeuralNetworkInference/Source/NeuralNetworkInference/Private/NeuralNetworkImplBackEndUEAndORT.cpp
+++ b/Engine/Plugins/Experimental/NeuralNetworkInference/Source/NeuralNetworkInference/Private/NeuralNetworkImplBackEndUEAndORT.cpp
@@ -119,7 +119,6 @@ IDMLDevice* FPrivateImplBackEndUEAndORT::FDMLDeviceList::Add(ID3D12Device* Devic
 /* UNeuralNetwork public functions
  *****************************************************************************/
 
-//#if WITH_EDITOR
 //bool UNeuralNetwork::FImplBackEndUEAndORT::LoadFile(TSharedPtr<FImplBackEndUEAndORT>& InOutImplBackEndUEAndORT, TArray<bool>& OutAreInputTensorSizesVariable, const TArray<uint8>& InModelReadFromFileInBytes, const FString& InModelFullFilePath, const ENeuralDeviceType InDeviceType)
 //{
 //	// Initialize and configure InOutImplBackEndUEAndORT
@@ -185,7 +184,6 @@ IDMLDevice* FPrivateImplBackEndUEAndORT::FDMLDeviceList::Add(ID3D12Device* Devic
 //	UE_LOG(LogNeuralNetworkInference, Warning, TEXT("UNeuralNetwork::Load(): Unknown file format \"%s\" used."), *FileExtension);
 //	return false;
 //}
-//#endif //WITH_EDITOR
 
 bool UNeuralNetwork::FImplBackEndUEAndORT::Load(TSharedPtr<FImplBackEndUEAndORT>& InOutImplBackEndUEAndORT, TArray<bool>& OutAreInputTensorSizesVariable, const TArray<uint8>& InModelReadFromFileInBytes, const FString& InModelFullFilePath, const ENeuralDeviceType InDeviceType)
 {
diff --git a/Engine/Plugins/Experimental/NeuralNetworkInference/Source/NeuralNetworkInference/Private/NeuralNetworkImplBackEndUEAndORT.h b/Engine/Plugins/Experimental/NeuralNetworkInference/Source/NeuralNetworkInference/Private/NeuralNetworkImplBackEndUEAndORT.h
--- a/Engine/Plugins/Experimental/NeuralNetworkInference/Source/NeuralNetworkInference/Private/NeuralNetworkImplBackEndUEAndORT.h
+++ b/Engine/Plugins/Experimental/NeuralNetworkInference/Source/NeuralNetworkInference/Private/NeuralNetworkImplBackEndUEAndORT.h
@@ -37,10 +37,6 @@ public:
 	TArray<const char*> OutputTensorNames; /* Tensor names */
 #endif //WITH_UE_AND_ORT_SUPPORT
 
-//#if WITH_EDITOR
-//	static bool LoadFile(TSharedPtr<FImplBackEndUEAndORT>& InOutImplBackEndUEAndORT, TArray<bool>& OutAreInputTensorSizesVariable, const TArray<uint8>& InModelReadFromFileInBytes, const FString& InModelFullFilePath, const ENeuralDeviceType InDeviceType);
-//#endif //WITH_EDITOR
-
 	static bool Load(TSharedPtr<FImplBackEndUEAndORT>& InOutImplBackEndUEAndORT, TArray<bool>& OutAreInputTensorSizesVariable, const TArray<uint8>& InModelReadFromFileInBytes, const FString& InModelFullFilePath, const ENeuralDeviceType InDeviceType);
 	
 	void Run(const ENeuralNetworkSynchronousMode InSynchronousMode, const ENeuralDeviceType InInputDeviceType, const ENeuralDeviceType InOutputDeviceType);
diff --git a/Engine/Plugins/Experimental/NeuralNetworkInference/Source/NeuralNetworkInference/Private/NeuralNetworkImplBackEndUEOnly.cpp b/Engine/Plugins/Experimental/NeuralNetworkInference/Source/NeuralNetworkInference/Private/NeuralNetworkImplBackEndUEOnly.cpp
--- a/Engine/Plugins/Experimental/NeuralNetworkInference/Source/NeuralNetworkInference/Private/NeuralNetworkImplBackEndUEOnly.cpp
+++ b/Engine/Plugins/Experimental/NeuralNetworkInference/Source/NeuralNetworkInference/Private/NeuralNetworkImplBackEndUEOnly.cpp
@@ -2,15 +2,12 @@
 
 #include "NeuralNetworkImplBackEndUEOnly.h"
 #include "GraphProtoToNeuralNetworkConverter.h"
+#include "ModelProtoFileReader.h"
 #include "NeuralNetworkInferenceUtils.h"
 #include "RenderGraphBuilder.h"
 #include "RenderingThread.h"
 #include "RHI.h"
 
-#if WITH_EDITOR
-#include "ModelProtoFileReader.h"
-#endif //WITH_EDITOR
-
 
 
 /* FImplBackEndUEOnly public functions
@@ -18,7 +15,6 @@
 
 bool UNeuralNetwork::FImplBackEndUEOnly::Load(TSharedPtr<FImplBackEndUEOnly>& InOutImplBackEndUEOnly, const FString& InModelFilePath)
 {
-#if WITH_EDITOR
 	// Initialize InOutImplBackEndUEOnly
 	if (!InOutImplBackEndUEOnly.IsValid())
 	{
@@ -42,11 +38,6 @@ bool UNeuralNetwork::FImplBackEndUEOnly::Load(TSharedPtr<FImplBackEndUEOnly>& In
 	}
 	// Load successful
 	return true;
-
-#else //WITH_EDITOR
-	UE_LOG(LogNeuralNetworkInference, Warning, TEXT("UNeuralNetwork::FImplBackEndUEOnly::Load(): Only implemented for Editor mode."));
-	return false;
-#endif //WITH_EDITOR
 }
 
 //bool UNeuralNetwork::FImplBackEndUEOnly::Load(TSharedPtr<FImplBackEndUEOnly>& InOutImplBackEndUEOnly, FNeuralTensorManager& InTensorManager, const TArray<TSharedPtr<FNeuralOperator>>& InOperators)
diff --git a/Engine/Plugins/Experimental/NeuralNetworkInference/Source/NeuralNetworkInference/Private/NeuralTensor.cpp b/Engine/Plugins/Experimental/NeuralNetworkInference/Source/NeuralNetworkInference/Private/NeuralTensor.cpp
--- a/Engine/Plugins/Experimental/NeuralNetworkInference/Source/NeuralNetworkInference/Private/NeuralTensor.cpp
+++ b/Engine/Plugins/Experimental/NeuralNetworkInference/Source/NeuralNetworkInference/Private/NeuralTensor.cpp
@@ -467,7 +467,6 @@ bool FNeuralTensor::SetFromTensorProto(const FTensorProto* const InTensorProto,
 	// InTensorProto->ExternalData
 	else if (!InTensorProto->ExternalData.IsEmpty())
 	{
-#if WITH_EDITOR
 		UE_LOG(LogNeuralNetworkInference, Warning, TEXT("FNeuralTensor::SetFromTensorProto(): DEPRECATED CODE, OTXT no longer functional."));
 		// // Sanity check
 		// if (InTensorProto->ExternalData[0].Key != TEXT("location") || InTensorProto->ExternalData[0].Value.Len() < 1)
@@ -483,10 +482,6 @@ bool FNeuralTensor::SetFromTensorProto(const FTensorProto* const InTensorProto,
 		// 	UE_LOG(LogNeuralNetworkInference, Warning, TEXT("FNeuralTensor::SetFromTensorProto(): Could not read binary file: %s."), *BinaryWeightFilePath);
 		// 	return false;
 		// }
-#else //WITH_EDITOR
-		UE_LOG(LogNeuralNetworkInference, Warning, TEXT("FNeuralTensor::SetFromTensorProto(): InTensorProto->ExternalData should never be used in non-Editor mode."));
-		return false;
-#endif //WITH_EDITOR
 	}
 	else
 	{
diff --git a/Engine/Plugins/Experimental/NeuralNetworkInference/Source/NeuralNetworkInference/Public/NeuralNetwork.h b/Engine/Plugins/Experimental/NeuralNetworkInference/Source/NeuralNetworkInference/Public/NeuralNetwork.h
--- a/Engine/Plugins/Experimental/NeuralNetworkInference/Source/NeuralNetworkInference/Public/NeuralNetwork.h
+++ b/Engine/Plugins/Experimental/NeuralNetworkInference/Source/NeuralNetworkInference/Public/NeuralNetwork.h
@@ -23,7 +23,6 @@ enum class ENeuralNetworkSynchronousMode : uint8
  *
  * See the following examples to learn how to read any ONNX model and run inference (i.e., a forward pass) on it.
  * 1. Constructing a UNeuralNetwork from an ONNX file (Editor-only):
- *	#if WITH_EDITOR
  *		// Create the UNeuralNetwork object
  *		UNeuralNetwork* Network = NewObject<UNeuralNetwork>((UObject*)GetTransientPackage(), UNeuralNetwork::StaticClass());
  *		// Try to load the network and set the device (CPU/GPU)
@@ -32,7 +31,6 @@ enum class ENeuralNetworkSynchronousMode : uint8
  *		{
  *			Network->SetDeviceType(ENeuralDeviceType::CPU); // Set to CPU/GPU mode
  *		}
- *	#endif
  *
  * 2. Loading a UNeuralNetwork from a previously-created UAsset (in Editor or in Game):
  *		// Create and load the UNeuralNetwork object from a UAsset
diff --git a/Engine/Plugins/Experimental/NeuralNetworkInference/Source/Private/ModelProtoConverter/ModelProtoConverter.Build.cs b/Engine/Plugins/Experimental/NeuralNetworkInference/Source/Private/ModelProtoConverter/ModelProtoConverter.Build.cs
--- a/Engine/Plugins/Experimental/NeuralNetworkInference/Source/Private/ModelProtoConverter/ModelProtoConverter.Build.cs
+++ b/Engine/Plugins/Experimental/NeuralNetworkInference/Source/Private/ModelProtoConverter/ModelProtoConverter.Build.cs
@@ -7,9 +7,10 @@ public class ModelProtoConverter : ModuleRules
 {
 	public ModelProtoConverter( ReadOnlyTargetRules Target ) : base( Target )
 	{
-		PCHUsage = ModuleRules.PCHUsageMode.UseExplicitOrSharedPCHs;
+		// Define when ModelProtoConverter is available
+		bool bIsModelProtoConverterSupported = (Target.Platform == UnrealTargetPlatform.Win64 || Target.Platform == UnrealTargetPlatform.Linux || Target.Platform == UnrealTargetPlatform.Mac);
 
-		PrivateIncludePaths.Add(Path.Combine(ModuleDirectory, "Generated"));
+		PCHUsage = ModuleRules.PCHUsageMode.UseExplicitOrSharedPCHs;
 
 		PublicIncludePaths.AddRange(
 			new string[] {
@@ -17,20 +18,31 @@ public class ModelProtoConverter : ModuleRules
 			}
 		);
 
+		PrivateIncludePaths.Add(Path.Combine(ModuleDirectory, "Generated"));
+
 		PublicDependencyModuleNames.AddRange
 			(
 			new string[] {
 				"Core",
 				"ModelProto"
 			}
-        );
+		);
 
 		PrivateDependencyModuleNames.AddRange
 			(
 			new string[] {
-				"Protobuf",
-                "ThirdPartyHelperAndDLLLoader"
-            }
+				"ThirdPartyHelperAndDLLLoader"
+			}
 		);
+
+		if (bIsModelProtoConverterSupported)
+		{
+			PrivateDependencyModuleNames.AddRange
+				(
+				new string[] {
+					"Protobuf"
+				}
+			);
+		}
 	}
 }
\ No newline at end of file
diff --git a/Engine/Plugins/Experimental/NeuralNetworkInference/Source/Private/ModelProtoConverter/Private/ModelProtoConverter.cpp b/Engine/Plugins/Experimental/NeuralNetworkInference/Source/Private/ModelProtoConverter/Private/ModelProtoConverter.cpp
--- a/Engine/Plugins/Experimental/NeuralNetworkInference/Source/Private/ModelProtoConverter/Private/ModelProtoConverter.cpp
+++ b/Engine/Plugins/Experimental/NeuralNetworkInference/Source/Private/ModelProtoConverter/Private/ModelProtoConverter.cpp
@@ -3,22 +3,23 @@
 #include "ModelProtoConverter.h"
 #include "ModelProtoConverterUtils.h"
 
+#if WITH_EDITOR // @todo: Workaround to avoid linking error when compiled in-game with UEAndORT back end, Build.cs file also modified
+
 // Protobuf includes
 #ifdef WITH_PROTOBUF
 #include "ThirdPartyWarningDisabler.h"
 NNI_THIRD_PARTY_INCLUDES_START
 #undef TEXT
 #undef check
-#include "onnx.proto3.pb.h"
+#include "onnx.proto3.pb.h" // From this module, but it causes conflicts when cooking the game and also using ONNXRuntime/ONNX
+// #include "onnx/onnx.pb.h" // From "Proto" module, if we want to avoid linking errors when cooking the game and also using ONNXRuntime/ONNX
 NNI_THIRD_PARTY_INCLUDES_END
-#endif //WITH_PROTOBUF
 
 
 
 /* FPrivateModelProtoConverter public functions
 *****************************************************************************/
 
-#ifdef WITH_PROTOBUF
 
 class FPrivateModelProtoConverter
 {
@@ -175,13 +176,7 @@ bool FPrivateModelProtoConverter::ConvertProto3ToUAsset(FGraphProto& OutGraphPro
 		UE_LOG(LogModelProtoConverter, Warning, TEXT("FPrivateModelProtoConverter::ConvertProto3ToUAsset(): ConvertProto3ToUAsset(TensorAnnotation) failed."));
 		return false;
 	}
-
-	if (false /* If something went wrong at any point, this is a template warning message*/)
-	{
-		UE_LOG(LogModelProtoConverter, Warning, TEXT("FPrivateModelProtoConverter::ConvertProto3ToUAsset(): FGraphProto could not be read from onnx::GraphProto."));
-		return false;
-	}
-
+	// No warnings occurred, conversion successful
 	return true;
 }
 
@@ -487,6 +482,7 @@ bool FPrivateModelProtoConverter::ConvertProto3ToUAssetProtoArrays(TArray<T1>& O
 }
 
 #endif //WITH_PROTOBUF
+#endif //WITH_EDITOR
 
 
 
@@ -495,13 +491,20 @@ bool FPrivateModelProtoConverter::ConvertProto3ToUAssetProtoArrays(TArray<T1>& O
 
 bool FModelProtoConverter::ConvertFromONNXProto3Ifstream(FModelProto& OutModelProto, std::istream& InIfstream)
 {
-#ifdef WITH_PROTOBUF
-	// Protobuf onnx::ModelProto from Ifstream
-	onnx::ModelProto ModelProto;
-	ModelProto.ParseFromIstream(&InIfstream);
-	return FPrivateModelProtoConverter::ConvertProto3ToUAsset(OutModelProto, ModelProto);
-#else //WITH_PROTOBUF
-	UE_LOG(LogModelProtoConverter, Warning, TEXT("FModelProtoConverter::ConvertFromONNXProto3Ifstream(): WITH_PROTOBUF was not defined."));
-	return false
-#endif //WITH_PROTOBUF
+#if WITH_EDITOR // @todo: Workaround to avoid linking error when compiled in-game with UEAndORT back end, Schema*.txt files also modified
+	#ifdef WITH_PROTOBUF
+		// Protobuf onnx::ModelProto from Ifstream
+		onnx::ModelProto ModelProto;
+		ModelProto.ParseFromIstream(&InIfstream);
+		return FPrivateModelProtoConverter::ConvertProto3ToUAsset(OutModelProto, ModelProto);
+	#else //WITH_PROTOBUF
+		UE_LOG(LogModelProtoConverter, Warning, TEXT("FModelProtoConverter::ConvertFromONNXProto3Ifstream(): WITH_PROTOBUF was not defined."));
+		return false;
+	#endif //WITH_PROTOBUF
+#else //WITH_EDITOR
+	UE_LOG(LogModelProtoConverter, Warning,
+		TEXT("FModelProtoConverter::ConvertFromONNXProto3Ifstream(): Due to a known linking error when used together with the UEAndORT back end during game mode,"
+			" this function is only available on Editor mode (WITH_EDITOR) for now."));
+	return false;
+#endif //WITH_EDITOR
 }
diff --git a/Engine/Plugins/Experimental/NeuralNetworkInference/Source/Private/ModelProtoConverter/SchemaFooter.txt b/Engine/Plugins/Experimental/NeuralNetworkInference/Source/Private/ModelProtoConverter/SchemaFooter.txt
--- a/Engine/Plugins/Experimental/NeuralNetworkInference/Source/Private/ModelProtoConverter/SchemaFooter.txt
+++ b/Engine/Plugins/Experimental/NeuralNetworkInference/Source/Private/ModelProtoConverter/SchemaFooter.txt
@@ -1,3 +1,5 @@
 
 NNI_THIRD_PARTY_INCLUDES_END
-#endif
\ No newline at end of file
+#endif //WITH_PROTOBUF
+
+#endif //WITH_EDITOR
\ No newline at end of file
diff --git a/Engine/Plugins/Experimental/NeuralNetworkInference/Source/Private/ModelProtoConverter/SchemaHeader.txt b/Engine/Plugins/Experimental/NeuralNetworkInference/Source/Private/ModelProtoConverter/SchemaHeader.txt
--- a/Engine/Plugins/Experimental/NeuralNetworkInference/Source/Private/ModelProtoConverter/SchemaHeader.txt
+++ b/Engine/Plugins/Experimental/NeuralNetworkInference/Source/Private/ModelProtoConverter/SchemaHeader.txt
@@ -1,3 +1,5 @@
+#if WITH_EDITOR
+
 //-V::568,590
 #ifdef WITH_PROTOBUF
 #include "CoreMinimal.h"
