NeuralNetworkInference (NNI) plugin:
- Some additional cleanup to the previous big commit.

_codereview Francisco.Vicente
#rnx

diff --git a/Engine/Plugins/Experimental/NeuralNetworkInference/Source/NeuralNetworkInferenceQA/NeuralNetworkInferenceQA.Build.cs b/Engine/Plugins/Experimental/NeuralNetworkInference/Source/NeuralNetworkInferenceQA/NeuralNetworkInferenceQA.Build.cs
--- a/Engine/Plugins/Experimental/NeuralNetworkInference/Source/NeuralNetworkInferenceQA/NeuralNetworkInferenceQA.Build.cs
+++ b/Engine/Plugins/Experimental/NeuralNetworkInference/Source/NeuralNetworkInferenceQA/NeuralNetworkInferenceQA.Build.cs
@@ -29,19 +29,9 @@ public class NeuralNetworkInferenceQA : ModuleRules
 			(
 			new string[] {
 				"ModelProto",
+				"ModelProtoFileReader",
 				"NeuralNetworkInferenceCore"
 			}
 		);
-
-		// Editor-only
-		if (Target.bBuildEditor)
-		{
-			PrivateDependencyModuleNames.AddRange
-				(
-				new string[] {
-					"ModelProtoFileReader"
-				}
-			);
-		}
 	}
 }
diff --git a/Engine/Plugins/Experimental/NeuralNetworkInference/Source/NeuralNetworkInferenceQA/Private/LegacyModelUnitTester.cpp b/Engine/Plugins/Experimental/NeuralNetworkInference/Source/NeuralNetworkInferenceQA/Private/LegacyModelUnitTester.cpp
--- a/Engine/Plugins/Experimental/NeuralNetworkInference/Source/NeuralNetworkInferenceQA/Private/LegacyModelUnitTester.cpp
+++ b/Engine/Plugins/Experimental/NeuralNetworkInference/Source/NeuralNetworkInferenceQA/Private/LegacyModelUnitTester.cpp
@@ -126,7 +126,7 @@ UNeuralNetworkLegacy* FLegacyModelUnitTester::NetworkONNXLoadTest(const FString&
 		return nullptr;
 	}
 	return Network;
-#else
+#else //WITH_EDITOR
 	UE_LOG(LogNeuralNetworkInferenceQA, Display, TEXT("--------------- NetworkONNXLoadTest test skipped (only in Editor)."));
 	return nullptr;
 #endif //WITH_EDITOR
diff --git a/Engine/Plugins/Experimental/NeuralNetworkInference/Source/NeuralNetworkInferenceQA/Private/ModelProtoFileReaderTester.cpp b/Engine/Plugins/Experimental/NeuralNetworkInference/Source/NeuralNetworkInferenceQA/Private/ModelProtoFileReaderTester.cpp
--- a/Engine/Plugins/Experimental/NeuralNetworkInference/Source/NeuralNetworkInferenceQA/Private/ModelProtoFileReaderTester.cpp
+++ b/Engine/Plugins/Experimental/NeuralNetworkInference/Source/NeuralNetworkInferenceQA/Private/ModelProtoFileReaderTester.cpp
@@ -1,11 +1,8 @@
 // Copyright Epic Games, Inc. All Rights Reserved.
 
 #include "ModelProtoFileReaderTester.h"
-#include "NeuralNetworkInferenceQAUtils.h"
-
-#if WITH_EDITOR
 #include "ModelProtoFileReader.h"
-#endif //WITH_EDITOR
+#include "NeuralNetworkInferenceQAUtils.h"
 
 
 
@@ -23,7 +20,7 @@ bool FModelProtoFileReaderTester::Test(const FString& InONNXModelFileName)
 	ensureMsgf(ModelProto.IsLoaded(), TEXT("FModelProtoFileReaderTester::ONNXReadNetworkTest() failed, FModelProto could not be read from InONNXModelFileName: %s."), *InONNXModelFileName);
 	return ModelProto.IsLoaded();
 #else //WITH_EDITOR
-	UE_LOG(LogNeuralNetworkInferenceQA, Display, TEXT("ONNXReadNetworkTest test skipped (only in Editor)."));
+	UE_LOG(LogNeuralNetworkInferenceQA, Display, TEXT("ONNXReadNetworkTest test skipped (only if WITH_EDITOR enabled)."));
 	return true;
 #endif //WITH_EDITOR
 }
diff --git a/Engine/Plugins/Experimental/NeuralNetworkInference/Source/NeuralNetworkInferenceQA/Private/ModelUnitTester.cpp b/Engine/Plugins/Experimental/NeuralNetworkInference/Source/NeuralNetworkInferenceQA/Private/ModelUnitTester.cpp
--- a/Engine/Plugins/Experimental/NeuralNetworkInference/Source/NeuralNetworkInferenceQA/Private/ModelUnitTester.cpp
+++ b/Engine/Plugins/Experimental/NeuralNetworkInference/Source/NeuralNetworkInferenceQA/Private/ModelUnitTester.cpp
@@ -94,8 +94,8 @@ if (ONNXOrORTIndex == 1)
 			UE_LOG(LogNeuralNetworkInferenceQA, Display, TEXT("-------------------- %s - Network %s Load and Run - %s"), *ModelName, *ModelType, *ModelFilePath);
 			bDidGlobalTestPassed &= ModelAccuracyTest(NetworkONNXOrORTLoadTest(ModelFilePath), InInputArrayValues, CPUGroundTruths, GPUGroundTruths);
 		}
-#else
-	UE_LOG(LogNeuralNetworkInferenceQA, Display, TEXT("-------------------- Skipped (only in Editor)."));
+#else //WITH_EDITOR
+	UE_LOG(LogNeuralNetworkInferenceQA, Display, TEXT("-------------------- Skipped (only in WITH_EDITOR enabled)."));
 #endif //WITH_EDITOR
 	}
 
@@ -142,7 +142,6 @@ UNeuralNetwork* FModelUnitTester::NetworkUassetLoadTest(const FString& InUAssetP
 	return Network;
 }
 
-#if WITH_EDITOR
 UNeuralNetwork* FModelUnitTester::NetworkONNXOrORTLoadTest(const FString& InModelFilePath)
 {
 	// Load network architecture and weights from file
@@ -159,7 +158,6 @@ UNeuralNetwork* FModelUnitTester::NetworkONNXOrORTLoadTest(const FString& InMode
 	}
 	return Network;
 }
-#endif //WITH_EDITOR
 
 bool FModelUnitTester::ModelAccuracyTest(UNeuralNetwork* InOutNetwork, const TArray<float>& InInputArrayValues, const TArray<double>& InCPUGroundTruths, const TArray<double>& InGPUGroundTruths)
 {
diff --git a/Engine/Plugins/Experimental/NeuralNetworkInference/Source/NeuralNetworkInferenceQA/Private/ModelUnitTester.h b/Engine/Plugins/Experimental/NeuralNetworkInference/Source/NeuralNetworkInferenceQA/Private/ModelUnitTester.h
--- a/Engine/Plugins/Experimental/NeuralNetworkInference/Source/NeuralNetworkInferenceQA/Private/ModelUnitTester.h
+++ b/Engine/Plugins/Experimental/NeuralNetworkInference/Source/NeuralNetworkInferenceQA/Private/ModelUnitTester.h
@@ -24,9 +24,7 @@ private:
 	static FString GetORTModelFilePath(const FString& InModelZooDirectory, const FString& InModelName);
 	static FString GetUAssetModelFilePath(const FString& InModelName, const FString& InModelZooRelativeDirectory);
 	static UNeuralNetwork* NetworkUassetLoadTest(const FString& InUAssetPath);
-#if WITH_EDITOR
 	static UNeuralNetwork* NetworkONNXOrORTLoadTest(const FString& InModelFilePath);
-#endif //WITH_EDITOR
 	static bool ModelAccuracyTest(UNeuralNetwork* InOutNetwork, const TArray<float>& InInputArrayValues, const TArray<double>& InCPUGroundTruths,
 		const TArray<double>& InGPUGroundTruths);
 	static bool ModelSpeedTest(const FString& InUAssetPath, const int32 InCPURepetitions, const int32 InGPURepetitions);
diff --git a/Engine/Plugins/Experimental/NeuralNetworkInference/Source/Private/ModelProtoConverter/ModelProtoConverter.Build.cs b/Engine/Plugins/Experimental/NeuralNetworkInference/Source/Private/ModelProtoConverter/ModelProtoConverter.Build.cs
--- a/Engine/Plugins/Experimental/NeuralNetworkInference/Source/Private/ModelProtoConverter/ModelProtoConverter.Build.cs
+++ b/Engine/Plugins/Experimental/NeuralNetworkInference/Source/Private/ModelProtoConverter/ModelProtoConverter.Build.cs
@@ -9,6 +9,10 @@ public class ModelProtoConverter : ModuleRules
 	{
 		// Define when ModelProtoConverter is available
 		bool bIsModelProtoConverterSupported = (Target.Platform == UnrealTargetPlatform.Win64 || Target.Platform == UnrealTargetPlatform.Linux || Target.Platform == UnrealTargetPlatform.Mac);
+		// if (bIsModelProtoConverterSupported)
+		// {
+		// 	PublicDefinitions.Add("WITH_MODEL_PROTO_CONVERTER_SUPPORT");
+		// }
 
 		PCHUsage = ModuleRules.PCHUsageMode.UseExplicitOrSharedPCHs;
 
@@ -31,6 +35,7 @@ public class ModelProtoConverter : ModuleRules
 		PrivateDependencyModuleNames.AddRange
 			(
 			new string[] {
+				// "ONNXRuntimeProto",
 				"ThirdPartyHelperAndDLLLoader"
 			}
 		);
diff --git a/Engine/Plugins/Experimental/NeuralNetworkInference/Source/Private/ModelProtoConverter/Private/ModelProtoConverter.cpp b/Engine/Plugins/Experimental/NeuralNetworkInference/Source/Private/ModelProtoConverter/Private/ModelProtoConverter.cpp
--- a/Engine/Plugins/Experimental/NeuralNetworkInference/Source/Private/ModelProtoConverter/Private/ModelProtoConverter.cpp
+++ b/Engine/Plugins/Experimental/NeuralNetworkInference/Source/Private/ModelProtoConverter/Private/ModelProtoConverter.cpp
@@ -498,7 +498,7 @@ bool FModelProtoConverter::ConvertFromONNXProto3Ifstream(FModelProto& OutModelPr
 		ModelProto.ParseFromIstream(&InIfstream);
 		return FPrivateModelProtoConverter::ConvertProto3ToUAsset(OutModelProto, ModelProto);
 	#else //WITH_PROTOBUF
-		UE_LOG(LogModelProtoConverter, Warning, TEXT("FModelProtoConverter::ConvertFromONNXProto3Ifstream(): WITH_PROTOBUF was not defined."));
+		UE_LOG(LogModelProtoConverter, Warning, TEXT("FModelProtoConverter::ConvertFromONNXProto3Ifstream(): Platform not compatible (WITH_PROTOBUF was not defined)."));
 		return false;
 	#endif //WITH_PROTOBUF
 #else //WITH_EDITOR
