NeuralNetworkInference (NNI) plugin:
- ProtoBuf converter able to take TArray<uint8> as input, but not connected properly to NeuralNetwork.cpp yet.

#rnx

diff --git a/Engine/Plugins/Experimental/NeuralNetworkInference/Source/NeuralNetworkInference/Private/GraphProtoToNeuralNetworkConverter.cpp b/Engine/Plugins/Experimental/NeuralNetworkInference/Source/NeuralNetworkInference/Private/GraphProtoToNeuralNetworkConverter.cpp
--- a/Engine/Plugins/Experimental/NeuralNetworkInference/Source/NeuralNetworkInference/Private/GraphProtoToNeuralNetworkConverter.cpp
+++ b/Engine/Plugins/Experimental/NeuralNetworkInference/Source/NeuralNetworkInference/Private/GraphProtoToNeuralNetworkConverter.cpp
@@ -12,116 +12,119 @@
 /* FGraphProtoToNeuralNetworkConverter public functions
  *****************************************************************************/
 
-bool FGraphProtoToNeuralNetworkConverter::Translate(TArray<TSharedPtr<FNeuralOperator>>& OutOperators, FNeuralTensorManager& InTensorManager, const FGraphProto& InGraphProto)
+bool FGraphProtoToNeuralNetworkConverter::Translate(TArray<TSharedPtr<FNeuralOperator>>& OutOperators, FNeuralTensorManager& InOrOutTensorManager, const FGraphProto& InGraphProto, const bool bInIsTensorManagerConst)
 {
-	TMap<FString, int32> NameIndexMap = InTensorManager.GetNameIndexMap();
-	TMap<FString, int32> DummyOutputNameIndexMap;
-	const bool bResult = CreateOperatorsAndEditTensorArray(OutOperators, DummyOutputNameIndexMap, InTensorManager.GetTensorsMutable(), NameIndexMap, InTensorManager.GetInputNameIndexMap(),
-		InTensorManager.GetOutputNameIndexMap(), InGraphProto);
-	if (NameIndexMap.Num() != InTensorManager.GetNameIndexMap().Num())
+	// InOrOutTensorManager is const, i.e., InTensorManager
+	if (bInIsTensorManagerConst)
 	{
-		UE_LOG(LogNeuralNetworkInference, Warning, TEXT("FGraphProtoToNeuralNetworkConverter::Translate(): NameIndexMap changed its size, it must not (%d != %d)."),
-			NameIndexMap.Num(), InTensorManager.GetNameIndexMap().Num());
+		TMap<FString, int32> NameIndexMap = InOrOutTensorManager.GetNameIndexMap();
+		TMap<FString, int32> DummyOutputNameIndexMap;
+		const bool bResult = CreateOperatorsAndEditTensorArray(OutOperators, DummyOutputNameIndexMap, InOrOutTensorManager.GetTensorsMutable(), NameIndexMap, InOrOutTensorManager.GetInputNameIndexMap(),
+			InOrOutTensorManager.GetOutputNameIndexMap(), InGraphProto, bInIsTensorManagerConst);
+		if (NameIndexMap.Num() != InOrOutTensorManager.GetNameIndexMap().Num())
+		{
+			UE_LOG(LogNeuralNetworkInference, Warning, TEXT("FGraphProtoToNeuralNetworkConverter::Translate(): NameIndexMap changed its size, it must not (%d != %d)."),
+				NameIndexMap.Num(), InOrOutTensorManager.GetNameIndexMap().Num());
+		}
+		return bResult;
 	}
-	return bResult;
-}
-
-bool FGraphProtoToNeuralNetworkConverter::Translate(TArray<TSharedPtr<FNeuralOperator>>& OutOperators, FNeuralTensorManager& OutTensorManager, const FGraphProto& InGraphProto,
-	const FString& InModelPath)
-{
-	// MaxNumberTensorsInNetwork is required to avoid the bug of resizing and invalidating all pointers given to the operators
-	int32 MaxNumberTensorsInNetwork = 0;
+	// InOrOutTensorManager is not const, i.e., OutTensorManager
+	else // if (!bInIsTensorManagerConst)
 	{
-		// Quickly estimate the number of FNeuralTensor's of the UNeuralNetworkLegacy for Tensors.Reserve(...)
-		// Number FNeuralTensor's = Number (unique) input FNeuralTensors + number (unique) output FNeuralTensors
-		// Uniqueness guaranteed by using the TSet
-		TSet<FString> TensorNameSet;
-		for (const FNodeProto& NodeProto : InGraphProto.Node)
+		// MaxNumberTensorsInNetwork is required to avoid the bug of resizing and invalidating all pointers given to the operators
+		int32 MaxNumberTensorsInNetwork = 0;
 		{
-			// Check input tensors
-			for (const FString& TensorName : NodeProto.Input)
+			// Quickly estimate the number of FNeuralTensor's of the UNeuralNetworkLegacy for Tensors.Reserve(...)
+			// Number FNeuralTensor's = Number (unique) input FNeuralTensors + number (unique) output FNeuralTensors
+			// Uniqueness guaranteed by using the TSet
+			TSet<FString> TensorNameSet;
+			for (const FNodeProto& NodeProto : InGraphProto.Node)
 			{
-				if (TensorNameSet.Find(TensorName) == nullptr)
+				// Check input tensors
+				for (const FString& TensorName : NodeProto.Input)
 				{
-					TensorNameSet.Add(TensorName);
+					if (TensorNameSet.Find(TensorName) == nullptr)
+					{
+						TensorNameSet.Add(TensorName);
+					}
 				}
-			}
-			// Check output tensors
-			for (const FString& TensorName : NodeProto.Output)
-			{
-				if (TensorNameSet.Find(TensorName) == nullptr)
+				// Check output tensors
+				for (const FString& TensorName : NodeProto.Output)
 				{
-					TensorNameSet.Add(TensorName);
+					if (TensorNameSet.Find(TensorName) == nullptr)
+					{
+						TensorNameSet.Add(TensorName);
+					}
 				}
 			}
+			MaxNumberTensorsInNetwork = TensorNameSet.Num()
+			// Add maximum number possible of auxiliary tensors. A conservative estimate is the number of operators times the maximum possible number of auxiliary tensors per operator
+				+ InGraphProto.Node.Num() * FNeuralOperator::GetMaximumPossibleNumberAuxiliaryTensorsPerOperator();
 		}
-		MaxNumberTensorsInNetwork = TensorNameSet.Num()
-		// Add maximum number possible of auxiliary tensors. A conservative estimate is the number of operators times the maximum possible number of auxiliary tensors per operator
-			+ InGraphProto.Node.Num() * FNeuralOperator::GetMaximumPossibleNumberAuxiliaryTensorsPerOperator();
-	}
-	TArray<FNeuralTensor> Tensors;
-	Tensors.Reserve(MaxNumberTensorsInNetwork);
-	TMap<FString, int32> NameIndexMap;
+		TArray<FNeuralTensor> Tensors;
+		Tensors.Reserve(MaxNumberTensorsInNetwork);
+		TMap<FString, int32> NameIndexMap;
 
-	// Read input tensors
-	TMap<FString, int32> InputNameIndexMap;
-	for (const FValueInfoProto& GraphProtoInput : InGraphProto.Input)
-	{
-		// Create tensor
-		const TArray<FTensorShapeProtoDimension>& Dimensions = GraphProtoInput.Type.TensorType.Shape.Dim;
-		TArray<int64> DimensionsAsTArray;
-		for (const FTensorShapeProtoDimension& Dimension : Dimensions)
+		// Read input tensors
+		TMap<FString, int32> InputNameIndexMap;
+		for (const FValueInfoProto& GraphProtoInput : InGraphProto.Input)
 		{
-			DimensionsAsTArray.Emplace(Dimension.DimValue);
+			// Create tensor
+			const TArray<FTensorShapeProtoDimension>& Dimensions = GraphProtoInput.Type.TensorType.Shape.Dim;
+			TArray<int64> DimensionsAsTArray;
+			for (const FTensorShapeProtoDimension& Dimension : Dimensions)
+			{
+				DimensionsAsTArray.Emplace(Dimension.DimValue);
+			}
+			// Add operator name to TMap
+			NameIndexMap.Add(GraphProtoInput.Name, Tensors.Num());
+			InputNameIndexMap.Add(GraphProtoInput.Name, Tensors.Num());
+			// Add tensor to TMap
+			Tensors.Emplace(FNeuralTensor(ENeuralDataType::Float, DimensionsAsTArray, GraphProtoInput.Name, ENeuralTensorTypeGPU::Input));
 		}
-		// Add operator name to TMap
-		NameIndexMap.Add(GraphProtoInput.Name, Tensors.Num());
-		InputNameIndexMap.Add(GraphProtoInput.Name, Tensors.Num());
-		// Add tensor to TMap
-		Tensors.Emplace(FNeuralTensor(ENeuralDataType::Float, DimensionsAsTArray, GraphProtoInput.Name, ENeuralTensorTypeGPU::Input));
-	}
 
-	// Read output tensor names (but do not create them yet, it has to happen simultaneously with the operator creation and tensor inlining)
-	// Issue: For inlining to work properly, output tensors cannot be created until after inlining operators (i.e., CreateOperatorsAndEditTensorArray).
-	// Otherwise, the following failing case could happen with a network like: InputTensor --> Gemm --> IntermediateTensor --> Relu --> OutputTensor
-	// 1. Inputs and outputs are created before hand, i.e., InputTensor and OutputTensor
-	// 2. CreateOperatorsAndEditTensorArray() is run
-	//   2.1. Gemm and IntermediateTensor would be created (because Gemm cannot be inlined)
-	//   2.2. Relu would be created, Relu would verify its input can be inlined, but both IntermediateTensor and OutputTensor are already created at this point, so it cannot inline them (crash/bug)
-	TMap<FString, int32> OutputNameDummyIndexMap;
-	for (const FValueInfoProto& GraphProtoOutput : InGraphProto.Output)
-	{
-		// Add operator name to TMap
-		if (NameIndexMap.Find(GraphProtoOutput.Name))
+		// Read output tensor names (but do not create them yet, it has to happen simultaneously with the operator creation and tensor inlining)
+		// Issue: For inlining to work properly, output tensors cannot be created until after inlining operators (i.e., CreateOperatorsAndEditTensorArray).
+		// Otherwise, the following failing case could happen with a network like: InputTensor --> Gemm --> IntermediateTensor --> Relu --> OutputTensor
+		// 1. Inputs and outputs are created before hand, i.e., InputTensor and OutputTensor
+		// 2. CreateOperatorsAndEditTensorArray() is run
+		//   2.1. Gemm and IntermediateTensor would be created (because Gemm cannot be inlined)
+		//   2.2. Relu would be created, Relu would verify its input can be inlined, but both IntermediateTensor and OutputTensor are already created at this point, so it cannot inline them (crash/bug)
+		TMap<FString, int32> OutputNameDummyIndexMap;
+		for (const FValueInfoProto& GraphProtoOutput : InGraphProto.Output)
 		{
-			UE_LOG(LogNeuralNetworkInference, Warning, TEXT("FGraphProtoToNeuralNetworkConverter::Translate(): A FNeuralTensor of the UNeuralNetworkLegacy is a global Input and Output simultaneously."));
+			// Add operator name to TMap
+			if (NameIndexMap.Find(GraphProtoOutput.Name))
+			{
+				UE_LOG(LogNeuralNetworkInference, Warning, TEXT("FGraphProtoToNeuralNetworkConverter::Translate(): A FNeuralTensor of the UNeuralNetworkLegacy is a global Input and Output simultaneously."));
+			}
+			OutputNameDummyIndexMap.Add(GraphProtoOutput.Name, -1);
 		}
-		OutputNameDummyIndexMap.Add(GraphProtoOutput.Name, -1);
-	}
 
-	// Fill operators
-	TMap<FString, int32> OutputNameIndexMap;
-	bool bWasSuccessful = CreateOperatorsAndEditTensorArray(OutOperators, OutputNameIndexMap, Tensors, NameIndexMap, InputNameIndexMap, OutputNameDummyIndexMap, InGraphProto, InModelPath);
-	if (Tensors.Num() <= NameIndexMap.Num())
-	{
-		UE_LOG(LogNeuralNetworkInference, Display, TEXT("Auto-inlining: %d initial tensors compressed into %d (%d tensors merged)."),
-			NameIndexMap.Num(), Tensors.Num(), NameIndexMap.Num() - Tensors.Num());
-	}
-	// Sanity check
-	if (Tensors.Num() > MaxNumberTensorsInNetwork)
-	{
-		UE_LOG(LogNeuralNetworkInference, Warning, TEXT("FGraphProtoToNeuralNetworkConverter::Translate(): Tensors.Num() <= MaxNumberTensorsInNetwork failed (%d vs. %d). A resize of Tensors"
-			" might have occurred, invalidating all TArray<FNeuralTensor*> FOperator variables. Report this error (so we can fix MaxNumberTensorsInNetwork to be big enough)."),
-			Tensors.Num(), MaxNumberTensorsInNetwork);
-		bWasSuccessful = false;
-	}
+		// Fill operators
+		TMap<FString, int32> OutputNameIndexMap;
+		bool bWasSuccessful = CreateOperatorsAndEditTensorArray(OutOperators, OutputNameIndexMap, Tensors, NameIndexMap, InputNameIndexMap, OutputNameDummyIndexMap, InGraphProto, bInIsTensorManagerConst);
+		if (Tensors.Num() <= NameIndexMap.Num())
+		{
+			UE_LOG(LogNeuralNetworkInference, Display, TEXT("Auto-inlining: %d initial tensors compressed into %d (%d tensors merged)."),
+				NameIndexMap.Num(), Tensors.Num(), NameIndexMap.Num() - Tensors.Num());
+		}
+		// Sanity check
+		if (Tensors.Num() > MaxNumberTensorsInNetwork)
+		{
+			UE_LOG(LogNeuralNetworkInference, Warning, TEXT("FGraphProtoToNeuralNetworkConverter::Translate(): Tensors.Num() <= MaxNumberTensorsInNetwork failed (%d vs. %d). A resize of Tensors"
+				" might have occurred, invalidating all TArray<FNeuralTensor*> FOperator variables. Report this error (so we can fix MaxNumberTensorsInNetwork to be big enough)."),
+				Tensors.Num(), MaxNumberTensorsInNetwork);
+			bWasSuccessful = false;
+		}
 
-	// Set OutTensorManager
-	if (bWasSuccessful)
-	{
-		return OutTensorManager.Load(Tensors, NameIndexMap, InputNameIndexMap, OutputNameIndexMap);
+		// Set InOrOutTensorManager
+		if (bWasSuccessful)
+		{
+			return InOrOutTensorManager.Load(Tensors, NameIndexMap, InputNameIndexMap, OutputNameIndexMap);
+		}
+		return bWasSuccessful;
 	}
-	return bWasSuccessful;
 }
 
 
@@ -131,9 +134,8 @@ bool FGraphProtoToNeuralNetworkConverter::Translate(TArray<TSharedPtr<FNeuralOpe
 
 bool FGraphProtoToNeuralNetworkConverter::CreateOperatorsAndEditTensorArray(TArray<TSharedPtr<FNeuralOperator>>& OutOperators, TMap<FString, int32>& OutputNameIndexMap,
 	TArray<FNeuralTensor>& InOutTensors, TMap<FString, int32>& InOutNameIndexMap, const TMap<FString, int32>& InInputNameDummyIndexMap,
-	const TMap<FString, int32>& InOutputNameDummyIndexMap, const FGraphProto& InGraphProto, const FString& InModelPath)
+	const TMap<FString, int32>& InOutputNameDummyIndexMap, const FGraphProto& InGraphProto, const bool bInIsTensorManagerConst)
 {
-	const bool bConvertingFromSerializedData = (InModelPath.Len() == 0);
 	const int64 InOutTensorsNumInit = InOutTensors.Num();
 	const int64 InOutNameIndexMapNumInit = InOutNameIndexMap.Num();
 	// Create operators sequentially
@@ -160,7 +162,7 @@ bool FGraphProtoToNeuralNetworkConverter::CreateOperatorsAndEditTensorArray(TArr
 				const FTensorProto& TensorProto = *FModelProto::FindElementInArray(InputTensorName, InGraphProto.Initializer, /*bMustValueBeFound*/true);
 				// Is it also an absolute output of the network?
 				const int64 TensorIndex = InOutTensors.Num();
-				if (!bConvertingFromSerializedData && InOutputNameDummyIndexMap.Find(InputTensorName))
+				if (!bInIsTensorManagerConst && InOutputNameDummyIndexMap.Find(InputTensorName))
 				{
 					OutputNameIndexMap.Add(InputTensorName, TensorIndex);
 				}
@@ -169,7 +171,7 @@ bool FGraphProtoToNeuralNetworkConverter::CreateOperatorsAndEditTensorArray(TArr
 				FNeuralTensor& NewTensor = InOutTensors.Last();
 				// Deprecated code for otxt files, if added back, add before the 2 for loops for speedup
 				// FString PathPart, FilenamePart, ExtensionPart;
-				// if (!bConvertingFromSerializedData)
+				// if (!bInIsTensorManagerConst)
 				// {
 				// 	FPaths::Split(InModelPath, PathPart, FilenamePart, ExtensionPart);
 				// }
@@ -194,8 +196,8 @@ bool FGraphProtoToNeuralNetworkConverter::CreateOperatorsAndEditTensorArray(TArr
 				return false;
 			}
 			const FString& OutputTensorName = NodeProto.Output[0];
-			// Case 1: bConvertingFromSerializedData
-			if (bConvertingFromSerializedData)
+			// Case 1: bInIsTensorManagerConst
+			if (bInIsTensorManagerConst)
 			{
 				// Sanity check
 				if (!InOutNameIndexMap.Find(OutputTensorName))
@@ -205,7 +207,7 @@ bool FGraphProtoToNeuralNetworkConverter::CreateOperatorsAndEditTensorArray(TArr
 				}
 				// Do nothing else
 			}
-			// Case 2: !bConvertingFromSerializedData
+			// Case 2: !bInIsTensorManagerConst
 			else
 			{
 				// Sanity check
@@ -243,7 +245,7 @@ bool FGraphProtoToNeuralNetworkConverter::CreateOperatorsAndEditTensorArray(TArr
 			// Create and configure operator
 			TSharedPtr<FNeuralOperator> Operator = CreateOperator(NodeProto, OperatorInputTensorNames, InInputNameDummyIndexMap, InOutputNameDummyIndexMap, InGraphProto);
 			// ConvTranspose requires flipping OperatorInputTensors[1], so we expect for now a Weight type, so it can be flipped beforehand
-			if (!bConvertingFromSerializedData && NodeProto.OperatorType == TEXT("ConvTranspose")) //if (NodeProto.OperatorType.Contains(TEXT("Conv"), ESearchCase::CaseSensitive))
+			if (!bInIsTensorManagerConst && NodeProto.OperatorType == TEXT("ConvTranspose")) //if (NodeProto.OperatorType.Contains(TEXT("Conv"), ESearchCase::CaseSensitive))
 			{
 				UE_LOG(LogNeuralNetworkInference, Display, TEXT("Flipping FConvTranspose weights. This is a temporary measurement but it will keep the weights flipped for now."));
 				// Sanity check
@@ -264,10 +266,10 @@ bool FGraphProtoToNeuralNetworkConverter::CreateOperatorsAndEditTensorArray(TArr
 			const int32 InlinedTensorIndex = Operator->SetInputTensorsAndGetInlinedIndex(OperatorInputTensors);
 			// Read/set output tensors
 			// 2 different cases for OperatorOutputTensors:
-				// Case 1/2: bConvertingFromSerializedData --> All FNeuralTensors already exist
+				// Case 1/2: bInIsTensorManagerConst --> All FNeuralTensors already exist
 					// Subcase 1/2: InlinedTensorIndex > -1 --> Do nothing
 					// Subcase 2/2: InlinedTensorIndex < 0 --> Fill OperatorOutputTensors from known index
-				// Case 2/2: !bConvertingFromSerializedData
+				// Case 2/2: !bInIsTensorManagerConst
 					// Subcase 1/3: ExistingTensorIndex is found --> FNeuralTensor already exists
 					// Subcase 2/3: InlinedTensorIndex < 0 --> Output tensor must be created because operator creates its own output
 					// Subcase 3/3: InlinedTensorIndex > -1 --> Output tensor is the same than 1 of the input tensors of that operator (inlined operator)
@@ -280,8 +282,8 @@ bool FGraphProtoToNeuralNetworkConverter::CreateOperatorsAndEditTensorArray(TArr
 			TArray<FNeuralTensor*> OperatorOutputTensors;
 			for (const FString& OutputTensorName : NodeProto.Output)
 			{
-				// Case 1: bConvertingFromSerializedData
-				if (bConvertingFromSerializedData)
+				// Case 1: bInIsTensorManagerConst
+				if (bInIsTensorManagerConst)
 				{
 					// Subcase 1/2: Do nothing
 					// Subcase 2/2: Fill OperatorOutputTensors from known index
@@ -291,10 +293,10 @@ bool FGraphProtoToNeuralNetworkConverter::CreateOperatorsAndEditTensorArray(TArr
 						OperatorOutputTensors.Push(&InOutTensors[ExistingTensorIndex]);
 					}
 				}
-				// Case 2: !bConvertingFromSerializedData
+				// Case 2: !bInIsTensorManagerConst
 				else
 				{
-					// Subcase 1/3: FNeuralTensor already exists (i.e., absolute output of the UNeuralNetworkLegacy or bConvertingFromSerializedData == true)
+					// Subcase 1/3: FNeuralTensor already exists (i.e., absolute output of the UNeuralNetworkLegacy or bInIsTensorManagerConst == true)
 					if (const int32* const ExistingTensorIndex = InOutNameIndexMap.Find(OutputTensorName))
 					{
 						if (InlinedTensorIndex > -1)
@@ -365,7 +367,7 @@ bool FGraphProtoToNeuralNetworkConverter::CreateOperatorsAndEditTensorArray(TArr
 				{
 					const FString AuxiliaryTensorName = NodeProto.Name + TEXT("-") + Operator->GetName() + TEXT("-Aux") + FString::FromInt(AuxiliaryTensorIndex);
 					// Uasset file
-					if (bConvertingFromSerializedData)
+					if (bInIsTensorManagerConst)
 					{
 						// Find and push auxiliary tensor
 						if (const int32* const AuxiliaryTensorAbsoluteIndex = InOutNameIndexMap.Find(AuxiliaryTensorName))
@@ -408,8 +410,8 @@ bool FGraphProtoToNeuralNetworkConverter::CreateOperatorsAndEditTensorArray(TArr
 			OutOperators.Push(Operator);
 		}
 	}
-	// Sanity checks for bConvertingFromSerializedData
-	if (bConvertingFromSerializedData)
+	// Sanity checks for bInIsTensorManagerConst
+	if (bInIsTensorManagerConst)
 	{
 		if (OutputNameIndexMap.Num() > 0)
 		{
@@ -433,7 +435,7 @@ bool FGraphProtoToNeuralNetworkConverter::CreateOperatorsAndEditTensorArray(TArr
 			return false;
 		}
 	}
-	// Sanity check for !bConvertingFromSerializedData
+	// Sanity check for !bInIsTensorManagerConst
 	else if (OutputNameIndexMap.Num() != InOutputNameDummyIndexMap.Num())
 	{
 		UE_LOG(LogNeuralNetworkInference, Warning,
diff --git a/Engine/Plugins/Experimental/NeuralNetworkInference/Source/NeuralNetworkInference/Private/GraphProtoToNeuralNetworkConverter.h b/Engine/Plugins/Experimental/NeuralNetworkInference/Source/NeuralNetworkInference/Private/GraphProtoToNeuralNetworkConverter.h
--- a/Engine/Plugins/Experimental/NeuralNetworkInference/Source/NeuralNetworkInference/Private/GraphProtoToNeuralNetworkConverter.h
+++ b/Engine/Plugins/Experimental/NeuralNetworkInference/Source/NeuralNetworkInference/Private/GraphProtoToNeuralNetworkConverter.h
@@ -17,11 +17,10 @@ public:
 	 * @param InTensorManager It is const, but kept not const for CreateOperatorsAndEditTensorArray() compatibility.
 	 *
 	 * 2 alternatives:
-	 * Without InModelPath, InOutTensorManager is assumed to be a const input variable (assuming the tensors are all already loaded, e.g., on de-serialization).
-	 * With InModelPath, InOutTensorManager is assumed to be an output variable that will be filled based on InModelPath.
+	 * If bInIsTensorManagerConst == true, InOrOutTensorManager is assumed to be a const input variable (assuming the tensors are all already loaded, e.g., on de-serialization).
+	 * With InModelPath, InOrOutTensorManager is assumed to be an output variable that will be filled based on InModelPath.
 	 */
-	static bool Translate(TArray<TSharedPtr<FNeuralOperator>>& OutOperators, FNeuralTensorManager& InTensorManager, const FGraphProto& InGraphProto);
-	static bool Translate(TArray<TSharedPtr<FNeuralOperator>>& OutOperators, FNeuralTensorManager& OutTensorManager, const FGraphProto& InGraphProto, const FString& InModelPath);
+	static bool Translate(TArray<TSharedPtr<FNeuralOperator>>& OutOperators, FNeuralTensorManager& InOrOutTensorManager, const FGraphProto& InGraphProto, const bool bInIsTensorManagerConst);
 
 private:
 	/**
@@ -30,8 +29,8 @@ private:
 	 * OutputNameIndexMap is just the copy of InOutputNameDummyIndexMap while filling its int32 Value fields. Given the dual Game/Editor behaviors, and to optimize Game mode, kept as 2 different rather than modifying in-place as InOutOutputNameDummyIndexMap.
 	 *
 	 * 2 alternatives:
-	 * - Game and/or Editor (efficient): If InModelPath is empty, InOutTensorManager is assumed to be a const input variable (assuming the tensors are all already loaded, e.g., on de-serialization).
-	 * - Editor-only: If InModelPath is not empty, InOutTensorManager is assumed to be an output variable that will be filled based on InModelPath.
+	 * - If bInIsTensorManagerConst == true, InOrOutTensorManager is assumed to be a const input variable (assuming the tensors are all already loaded, e.g., on de-serialization).
+	 * - If bInIsTensorManagerConst == false, InOrOutTensorManager is assumed to be an output variable that will be filled based on InModelPath.
 	 *
 	 * @param OutputNameIndexMap It is created from InOutputNameDummyIndexMap, assuming the FString part of the mapping is right, and it will fill the index (int32 part of it).
 	 * @param InInputNameDummyIndexMap The int32 part is irrelevant, the only important part is whether the FString field is filled or not, regardless of its value. Kept as TMap for performance to integrate better with FTensorManager.
@@ -39,7 +38,7 @@ private:
 	 */
 	static bool CreateOperatorsAndEditTensorArray(TArray<TSharedPtr<FNeuralOperator>>& OutOperators, TMap<FString, int32>& OutputNameIndexMap, TArray<FNeuralTensor>& InOutTensors,
 		TMap<FString, int32>& InOutNameIndexMap, const TMap<FString, int32>& InInputNameDummyIndexMap, const TMap<FString, int32>& InOutputNameDummyIndexMap,
-		const FGraphProto& InGraphProto, const FString& InModelPath = TEXT(""));
+		const FGraphProto& InGraphProto, const bool bInIsTensorManagerConst);
 
 	/**
 	 * It will return FNeuralOperator::EDataType::NotInlined in any of the 2 following cases occur:
diff --git a/Engine/Plugins/Experimental/NeuralNetworkInference/Source/NeuralNetworkInference/Private/NeuralNetwork.cpp b/Engine/Plugins/Experimental/NeuralNetworkInference/Source/NeuralNetworkInference/Private/NeuralNetwork.cpp
--- a/Engine/Plugins/Experimental/NeuralNetworkInference/Source/NeuralNetworkInference/Private/NeuralNetwork.cpp
+++ b/Engine/Plugins/Experimental/NeuralNetworkInference/Source/NeuralNetworkInference/Private/NeuralNetwork.cpp
@@ -131,7 +131,7 @@ bool UNeuralNetwork::Load()
 	// UEOnly
 	else if (BackEndForCurrentPlatform == ENeuralBackEnd::UEOnly)
 	{
-		bIsLoaded = UNeuralNetwork::FImplBackEndUEOnly::Load(ImplBackEndUEOnly, /*ModelReadFromFileInBytes*/ ModelFullFilePath);
+		bIsLoaded = UNeuralNetwork::FImplBackEndUEOnly::Load(ImplBackEndUEOnly, ModelReadFromFileInBytes, ModelFullFilePath);
 	}
 	// Unknown
 	else
diff --git a/Engine/Plugins/Experimental/NeuralNetworkInference/Source/NeuralNetworkInference/Private/NeuralNetworkImplBackEndUEOnly.cpp b/Engine/Plugins/Experimental/NeuralNetworkInference/Source/NeuralNetworkInference/Private/NeuralNetworkImplBackEndUEOnly.cpp
--- a/Engine/Plugins/Experimental/NeuralNetworkInference/Source/NeuralNetworkInference/Private/NeuralNetworkImplBackEndUEOnly.cpp
+++ b/Engine/Plugins/Experimental/NeuralNetworkInference/Source/NeuralNetworkInference/Private/NeuralNetworkImplBackEndUEOnly.cpp
@@ -13,7 +13,7 @@
 /* FImplBackEndUEOnly public functions
  *****************************************************************************/
 
-bool UNeuralNetwork::FImplBackEndUEOnly::Load(TSharedPtr<FImplBackEndUEOnly>& InOutImplBackEndUEOnly, const FString& InModelFilePath)
+bool UNeuralNetwork::FImplBackEndUEOnly::Load(TSharedPtr<FImplBackEndUEOnly>& InOutImplBackEndUEOnly, const TArray<uint8>& InModelReadFromFileInBytes, const FString& InModelFilePath)
 {
 	// Initialize InOutImplBackEndUEOnly
 	if (!InOutImplBackEndUEOnly.IsValid())
@@ -31,7 +31,7 @@ bool UNeuralNetwork::FImplBackEndUEOnly::Load(TSharedPtr<FImplBackEndUEOnly>& In
 		return false;
 	}
 	// Turn ModelProto into Operators
-	if (!FGraphProtoToNeuralNetworkConverter::Translate(InOutImplBackEndUEOnly->Operators, InOutImplBackEndUEOnly->TensorManager, InOutImplBackEndUEOnly->ModelProto.GetGraph(), InModelFilePath) || !InOutImplBackEndUEOnly->TensorManager.IsLoaded())
+	if (!FGraphProtoToNeuralNetworkConverter::Translate(InOutImplBackEndUEOnly->Operators, InOutImplBackEndUEOnly->TensorManager, InOutImplBackEndUEOnly->ModelProto.GetGraph(), /*bInIsTensorManagerConst*/false) || !InOutImplBackEndUEOnly->TensorManager.IsLoaded())
 	{
 		UE_LOG(LogNeuralNetworkInference, Warning, TEXT("UNeuralNetwork::FImplBackEndUEOnly::Load(): TensorManager could not be loaded from %s."), *InModelFilePath);
 		return false;
diff --git a/Engine/Plugins/Experimental/NeuralNetworkInference/Source/NeuralNetworkInference/Private/NeuralNetworkImplBackEndUEOnly.h b/Engine/Plugins/Experimental/NeuralNetworkInference/Source/NeuralNetworkInference/Private/NeuralNetworkImplBackEndUEOnly.h
--- a/Engine/Plugins/Experimental/NeuralNetworkInference/Source/NeuralNetworkInference/Private/NeuralNetworkImplBackEndUEOnly.h
+++ b/Engine/Plugins/Experimental/NeuralNetworkInference/Source/NeuralNetworkInference/Private/NeuralNetworkImplBackEndUEOnly.h
@@ -31,7 +31,7 @@ struct UNeuralNetwork::FImplBackEndUEOnly
 	 */
 	TArray<TSharedPtr<FNeuralOperator>> Operators;
 
- 	static bool Load(TSharedPtr<FImplBackEndUEOnly>& InOutImplBackEndUEOnly, const FString& InModelFullFilePath);
+ 	static bool Load(TSharedPtr<FImplBackEndUEOnly>& InOutImplBackEndUEOnly, const TArray<uint8>& InModelReadFromFileInBytes, const FString& InModelFullFilePath);
 
 	//static bool Load(TSharedPtr<FImplBackEndUEOnly>& InOutImplBackEndUEOnly, FNeuralTensorManager& InTensorManager, const TArray<TSharedPtr<FNeuralOperator>>& InOperators);
 
diff --git a/Engine/Plugins/Experimental/NeuralNetworkInference/Source/NeuralNetworkInference/Private/NeuralNetworkLegacy.cpp b/Engine/Plugins/Experimental/NeuralNetworkInference/Source/NeuralNetworkInference/Private/NeuralNetworkLegacy.cpp
--- a/Engine/Plugins/Experimental/NeuralNetworkInference/Source/NeuralNetworkInference/Private/NeuralNetworkLegacy.cpp
+++ b/Engine/Plugins/Experimental/NeuralNetworkInference/Source/NeuralNetworkInference/Private/NeuralNetworkLegacy.cpp
@@ -60,7 +60,7 @@ void UNeuralNetworkLegacy::PostLoad()
 			bIsLoaded = false;
 		}
 		// Turn ModelProto into Operators
-		else if (!FGraphProtoToNeuralNetworkConverter::Translate(Operators, TensorManager, ModelProto.GetGraph()))
+		else if (!FGraphProtoToNeuralNetworkConverter::Translate(Operators, TensorManager, ModelProto.GetGraph(), /*bInIsTensorManagerConst*/true))
 		{
 			UE_LOG(LogNeuralNetworkInference, Warning, TEXT("UNeuralNetworkLegacy::PostLoad(): UNeuralNetworkLegacy could not be configured from its FModelProto."));
 		}
@@ -131,7 +131,7 @@ bool UNeuralNetworkLegacy::Load(const FString& InModelFilePath)
 		return false;
 	}
 	// Turn ModelProto into Operators
-	bIsLoaded = FGraphProtoToNeuralNetworkConverter::Translate(Operators, TensorManager, ModelProto.GetGraph(), InModelFilePath);
+	bIsLoaded = FGraphProtoToNeuralNetworkConverter::Translate(Operators, TensorManager, ModelProto.GetGraph(), /*bInIsTensorManagerConst*/false);
 	if (!TensorManager.IsLoaded())
 	{
 		UE_LOG(LogNeuralNetworkInference, Warning, TEXT("UNeuralNetworkLegacy::Load(): TensorManager could not be loaded from %s."), *InModelFilePath);
diff --git a/Engine/Plugins/Experimental/NeuralNetworkInference/Source/Private/ModelProtoConverter/Internal/ModelProtoConverter.h b/Engine/Plugins/Experimental/NeuralNetworkInference/Source/Private/ModelProtoConverter/Internal/ModelProtoConverter.h
--- a/Engine/Plugins/Experimental/NeuralNetworkInference/Source/Private/ModelProtoConverter/Internal/ModelProtoConverter.h
+++ b/Engine/Plugins/Experimental/NeuralNetworkInference/Source/Private/ModelProtoConverter/Internal/ModelProtoConverter.h
@@ -10,8 +10,16 @@ class MODELPROTOCONVERTER_API FModelProtoConverter
 {
 public:
 	/**
-	 * It creates and fills OutModelProto from the *.onnx file read into InIfstream.
+	 * It creates and fills OutModelProto from the *.onnx file read into InIfstream or TArray<uint8>.
 	 * @return False if it could not initialize it, true if successful.
 	 */
 	static bool ConvertFromONNXProto3Ifstream(FModelProto& OutModelProto, std::istream& InIfstream);
+	static bool ConvertFromONNXProto3Array(FModelProto& OutModelProto, const TArray<uint8>& InModelReadFromFileInBytes);
+
+private:
+	/**
+	 * Auxiliary function to avoid code redundancies.
+	 * Exactly one of the 2 input arguments must be a nullptr.
+	 */
+	static bool ConvertFromONNXProto3(FModelProto& OutModelProto, std::istream* InIfstream, const TArray<uint8>* InModelReadFromFileInBytes);
 };
diff --git a/Engine/Plugins/Experimental/NeuralNetworkInference/Source/Private/ModelProtoConverter/Private/ModelProtoConverter.cpp b/Engine/Plugins/Experimental/NeuralNetworkInference/Source/Private/ModelProtoConverter/Private/ModelProtoConverter.cpp
--- a/Engine/Plugins/Experimental/NeuralNetworkInference/Source/Private/ModelProtoConverter/Private/ModelProtoConverter.cpp
+++ b/Engine/Plugins/Experimental/NeuralNetworkInference/Source/Private/ModelProtoConverter/Private/ModelProtoConverter.cpp
@@ -490,21 +490,56 @@ bool FPrivateModelProtoConverter::ConvertProto3ToUAssetProtoArrays(TArray<T1>& O
 *****************************************************************************/
 
 bool FModelProtoConverter::ConvertFromONNXProto3Ifstream(FModelProto& OutModelProto, std::istream& InIfstream)
+{
+	return ConvertFromONNXProto3(OutModelProto, &InIfstream, nullptr);
+}
+
+bool FModelProtoConverter::ConvertFromONNXProto3Array(FModelProto& OutModelProto, const TArray<uint8>& InModelReadFromFileInBytes)
+{
+	return ConvertFromONNXProto3(OutModelProto, nullptr, &InModelReadFromFileInBytes);
+}
+
+
+
+/* FModelProtoConverter private functions
+*****************************************************************************/
+
+bool FModelProtoConverter::ConvertFromONNXProto3(FModelProto& OutModelProto, std::istream* InIfstream, const TArray<uint8>* InModelReadFromFileInBytes)
 {
 #if WITH_EDITOR // @todo: Workaround to avoid linking error when compiled in-game with UEAndORT back end, Schema*.txt files also modified
 	#ifdef WITH_PROTOBUF
-		// Protobuf onnx::ModelProto from Ifstream
+		// Sanity checks
+		if ((InIfstream == nullptr) == (InModelReadFromFileInBytes == nullptr))
+		{
+			UE_LOG(LogModelProtoConverter, Warning, TEXT("FModelProtoConverter::ConvertFromONNXProto3(): Both InIfstream & InModelReadFromFileInBytes were nullptr or both were not nullptr. However, only 1 of them should not be a nullptr."));
+			return false;
+		}
+		// Protobuf onnx::ModelProto from Ifstream/InModelReadFromFileInBytes
 		onnx::ModelProto ModelProto;
-		ModelProto.ParseFromIstream(&InIfstream);
+		if (InIfstream)
+		{
+			ModelProto.ParseFromIstream(InIfstream);
+		}
+		else
+		{
+			if (InModelReadFromFileInBytes->IsEmpty())
+			{
+				UE_LOG(LogModelProtoConverter, Warning, TEXT("FModelProtoConverter::ConvertFromONNXProto3(): Input InModelReadFromFileInBytes was empty."));
+				return false;
+			}
+			ModelProto.ParseFromArray(InModelReadFromFileInBytes->GetData(), InModelReadFromFileInBytes->Num());
+		}
+		// FModelProto from Protobuf onnx::ModelProto
 		return FPrivateModelProtoConverter::ConvertProto3ToUAsset(OutModelProto, ModelProto);
+
 	#else //WITH_PROTOBUF
-		UE_LOG(LogModelProtoConverter, Warning, TEXT("FModelProtoConverter::ConvertFromONNXProto3Ifstream(): Platform not compatible (WITH_PROTOBUF was not defined)."));
+		UE_LOG(LogModelProtoConverter, Warning, TEXT("FModelProtoConverter::ConvertFromONNXProto3(): Platform not compatible (WITH_PROTOBUF was not defined)."));
 		return false;
 	#endif //WITH_PROTOBUF
+
 #else //WITH_EDITOR
 	UE_LOG(LogModelProtoConverter, Warning,
-		TEXT("FModelProtoConverter::ConvertFromONNXProto3Ifstream(): Due to a known linking error when used together with the UEAndORT back end during game mode,"
-			" this function is only available on Editor mode (WITH_EDITOR) for now."));
+		TEXT("FModelProtoConverter::ConvertFromONNXProto3(): Due to a known linking error when used together with the UEAndORT back end during game mode, this function is only available on Editor mode (WITH_EDITOR) for now."));
 	return false;
 #endif //WITH_EDITOR
 }
diff --git a/Engine/Plugins/Experimental/NeuralNetworkInference/Source/Private/ModelProtoFileReader/Internal/ModelProtoFileReader.h b/Engine/Plugins/Experimental/NeuralNetworkInference/Source/Private/ModelProtoFileReader/Internal/ModelProtoFileReader.h
--- a/Engine/Plugins/Experimental/NeuralNetworkInference/Source/Private/ModelProtoFileReader/Internal/ModelProtoFileReader.h
+++ b/Engine/Plugins/Experimental/NeuralNetworkInference/Source/Private/ModelProtoFileReader/Internal/ModelProtoFileReader.h
@@ -12,6 +12,7 @@ public:
 	 * It will return nullptr if it could not initialize it.
 	 */
 	static bool ReadModelProtoFromFile(FModelProto& OutModelProto, const FString& InFilePath);
+	static bool ReadModelProtoFromArray(FModelProto& OutModelProto, const TArray<uint8>& InModelReadFromFileInBytes);
 
 	/**
 	 * As long as InBinaryWeightsByteSize is bigger than the file size, it won't crash.
diff --git a/Engine/Plugins/Experimental/NeuralNetworkInference/Source/Private/ModelProtoFileReader/Private/ModelProtoFileReader.cpp b/Engine/Plugins/Experimental/NeuralNetworkInference/Source/Private/ModelProtoFileReader/Private/ModelProtoFileReader.cpp
--- a/Engine/Plugins/Experimental/NeuralNetworkInference/Source/Private/ModelProtoFileReader/Private/ModelProtoFileReader.cpp
+++ b/Engine/Plugins/Experimental/NeuralNetworkInference/Source/Private/ModelProtoFileReader/Private/ModelProtoFileReader.cpp
@@ -70,6 +70,11 @@ bool FModelProtoFileReader::ReadModelProtoFromFile(FModelProto& OutModelProto, c
 	return false;
 }
 
+bool FModelProtoFileReader::ReadModelProtoFromArray(FModelProto& OutModelProto, const TArray<uint8>& InModelReadFromFileInBytes)
+{
+	return FModelProtoConverter::ConvertFromONNXProto3Array(OutModelProto, InModelReadFromFileInBytes);
+}
+
 bool FModelProtoFileReader::ReadWeightsFromOtxtBinaryFile(char* OutBinaryWeights, const int32 InBinaryWeightsByteSize, const FString& InFilePath)
 {
 	// Note:
