NeuralNetworkInference (NNI) plugin:
- Towards UNN working with UEOnly backend: General cleaup and improved warning descriptions and removed FString ModelPath from UEOnly::Load().

#rnx

diff --git a/Engine/Plugins/Experimental/NeuralNetworkInference/Source/NeuralNetworkInference/Private/NeuralNetwork.cpp b/Engine/Plugins/Experimental/NeuralNetworkInference/Source/NeuralNetworkInference/Private/NeuralNetwork.cpp
--- a/Engine/Plugins/Experimental/NeuralNetworkInference/Source/NeuralNetworkInference/Private/NeuralNetwork.cpp
+++ b/Engine/Plugins/Experimental/NeuralNetworkInference/Source/NeuralNetworkInference/Private/NeuralNetwork.cpp
@@ -131,7 +131,7 @@ bool UNeuralNetwork::Load()
 	// UEOnly
 	else if (BackEndForCurrentPlatform == ENeuralBackEnd::UEOnly)
 	{
-		bIsLoaded = UNeuralNetwork::FImplBackEndUEOnly::Load(ImplBackEndUEOnly, ModelReadFromFileInBytes, ModelFullFilePath);
+		bIsLoaded = UNeuralNetwork::FImplBackEndUEOnly::Load(ImplBackEndUEOnly, ModelReadFromFileInBytes);
 	}
 	// Unknown
 	else
diff --git a/Engine/Plugins/Experimental/NeuralNetworkInference/Source/NeuralNetworkInference/Private/NeuralNetworkImplBackEndUEAndORT.cpp b/Engine/Plugins/Experimental/NeuralNetworkInference/Source/NeuralNetworkInference/Private/NeuralNetworkImplBackEndUEAndORT.cpp
--- a/Engine/Plugins/Experimental/NeuralNetworkInference/Source/NeuralNetworkInference/Private/NeuralNetworkImplBackEndUEAndORT.cpp
+++ b/Engine/Plugins/Experimental/NeuralNetworkInference/Source/NeuralNetworkInference/Private/NeuralNetworkImplBackEndUEAndORT.cpp
@@ -103,7 +103,7 @@ IDMLDevice* FPrivateImplBackEndUEAndORT::FDMLDeviceList::Add(ID3D12Device* Devic
 
 	if (!DmlDevice)
 	{
-		UE_LOG(LogNeuralNetworkInference, Error, TEXT("Failed to create DML device."));
+		UE_LOG(LogNeuralNetworkInference, Warning, TEXT("FDMLDeviceList::Add(): Failed to create DML device."));
 		return nullptr;
 	}
 
@@ -119,13 +119,13 @@ IDMLDevice* FPrivateImplBackEndUEAndORT::FDMLDeviceList::Add(ID3D12Device* Devic
 /* UNeuralNetwork public functions
  *****************************************************************************/
 
-//bool UNeuralNetwork::FImplBackEndUEAndORT::LoadFile(TSharedPtr<FImplBackEndUEAndORT>& InOutImplBackEndUEAndORT, TArray<bool>& OutAreInputTensorSizesVariable, const TArray<uint8>& InModelReadFromFileInBytes, const FString& InModelFullFilePath, const ENeuralDeviceType InDeviceType)
+//bool UNeuralNetwork::FImplBackEndUEAndORT::Load(TSharedPtr<FImplBackEndUEAndORT>& InOutImplBackEndUEAndORT, TArray<bool>& OutAreInputTensorSizesVariable, const TArray<uint8>& InModelReadFromFileInBytes, const FString& InModelFullFilePath, const ENeuralDeviceType InDeviceType)
 //{
 //	// Initialize and configure InOutImplBackEndUEAndORT
 //	const FRedirectCoutAndCerrToUeLog RedirectCoutAndCerrToUeLog;
 //	if (!InitializedAndConfigureMembers(InOutImplBackEndUEAndORT, InModelFullFilePath, InDeviceType))
 //	{
-//		UE_LOG(LogNeuralNetworkInference, Warning, TEXT("UNeuralNetwork::Load(): InitializedAndConfigureMembers failed."));
+//		UE_LOG(LogNeuralNetworkInference, Warning, TEXT("FImplBackEndUEAndORT::Load(): InitializedAndConfigureMembers failed."));
 //		return false;
 //	}
 //
@@ -181,7 +181,7 @@ IDMLDevice* FPrivateImplBackEndUEAndORT::FDMLDeviceList::Add(ID3D12Device* Devic
 //		return Load();
 //	}
 //	// Else
-//	UE_LOG(LogNeuralNetworkInference, Warning, TEXT("UNeuralNetwork::Load(): Unknown file format \"%s\" used."), *FileExtension);
+//	UE_LOG(LogNeuralNetworkInference, Warning, TEXT("FImplBackEndUEAndORT::Load(): Unknown file format \"%s\" used."), *FileExtension);
 //	return false;
 //}
 
@@ -197,7 +197,7 @@ bool UNeuralNetwork::FImplBackEndUEAndORT::Load(TSharedPtr<FImplBackEndUEAndORT>
 		// Initialize and configure InOutImplBackEndUEAndORT
 		if (!UNeuralNetwork::FImplBackEndUEAndORT::InitializedAndConfigureMembers(InOutImplBackEndUEAndORT, InModelFullFilePath, InDeviceType))
 		{
-			UE_LOG(LogNeuralNetworkInference, Warning, TEXT("UNeuralNetwork::Load(): InitializedAndConfigureMembers failed."));
+			UE_LOG(LogNeuralNetworkInference, Warning, TEXT("FImplBackEndUEAndORT::Load(): InitializedAndConfigureMembers failed."));
 			return false;
 		}
 
@@ -210,7 +210,7 @@ bool UNeuralNetwork::FImplBackEndUEAndORT::Load(TSharedPtr<FImplBackEndUEAndORT>
 		// Else
 		else
 		{
-			UE_LOG(LogNeuralNetworkInference, Warning, TEXT("InModelReadFromFileInBytes was empty."));
+			UE_LOG(LogNeuralNetworkInference, Warning, TEXT("FImplBackEndUEAndORT::Load(): InModelReadFromFileInBytes was empty."));
 			return false;
 		}
 
@@ -228,7 +228,7 @@ bool UNeuralNetwork::FImplBackEndUEAndORT::Load(TSharedPtr<FImplBackEndUEAndORT>
 #endif //WITH_EDITOR
 
 #else //WITH_UE_AND_ORT_SUPPORT
-	UE_LOG(LogNeuralNetworkInference, Warning, TEXT("UNeuralNetwork::Load(): Platform or Operating System not suported yet for UEAndORT BackEnd. Set BackEnd to ENeuralBackEnd::Auto (recommended) or ENeuralBackEnd::UEOnly for this platform."));
+	UE_LOG(LogNeuralNetworkInference, Warning, TEXT("FImplBackEndUEAndORT::Load(): Platform or Operating System not suported yet for UEAndORT BackEnd. Set BackEnd to ENeuralBackEnd::Auto (recommended) or ENeuralBackEnd::UEOnly for this platform."));
 	return false;
 #endif //WITH_UE_AND_ORT_SUPPORT
 }
@@ -245,7 +245,7 @@ void UNeuralNetwork::FImplBackEndUEAndORT::Run(const ENeuralNetworkSynchronousMo
 		// @todo: Temporarily disabled until we connect GPU input/output between UE and ORT
 		if (InInputDeviceType == ENeuralDeviceType::GPU || InOutputDeviceType == ENeuralDeviceType::GPU)
 		{
-			UE_LOG(LogNeuralNetworkInference, Warning, TEXT("UNeuralNetwork::Run(): InputDeviceType & OutputDeviceType must be set to CPU for now."));
+			UE_LOG(LogNeuralNetworkInference, Warning, TEXT("FImplBackEndUEAndORT::Run(): InputDeviceType & OutputDeviceType must be set to CPU for now."));
 			return;
 		}
 
@@ -258,11 +258,11 @@ void UNeuralNetwork::FImplBackEndUEAndORT::Run(const ENeuralNetworkSynchronousMo
 		}
 		else if (InSynchronousMode == ENeuralNetworkSynchronousMode::Asynchronous)
 		{
-			UE_LOG(LogNeuralNetworkInference, Warning, TEXT("UNeuralNetwork::Run(): SynchronousMode = %d not implemented yet. Use SynchronousMode = Synchronous."), (int32)InSynchronousMode);
+			UE_LOG(LogNeuralNetworkInference, Warning, TEXT("FImplBackEndUEAndORT::Run(): SynchronousMode = %d not implemented yet. Use SynchronousMode = Synchronous."), (int32)InSynchronousMode);
 		}
 		else
 		{
-			UE_LOG(LogNeuralNetworkInference, Warning, TEXT("UNeuralNetwork::Run(): Unknown SynchronousMode = %d."), (int32)InSynchronousMode);
+			UE_LOG(LogNeuralNetworkInference, Warning, TEXT("FImplBackEndUEAndORT::Run(): Unknown SynchronousMode = %d."), (int32)InSynchronousMode);
 		}
 	}
 #if WITH_EDITOR
@@ -273,7 +273,7 @@ void UNeuralNetwork::FImplBackEndUEAndORT::Run(const ENeuralNetworkSynchronousMo
 #endif //WITH_EDITOR
 
 #else //WITH_UE_AND_ORT_SUPPORT
-	UE_LOG(LogNeuralNetworkInference, Warning, TEXT("UNeuralNetwork::Run(): Platform or Operating System not suported yet for UEAndORT BackEnd. Set BackEnd to ENeuralBackEnd::Auto or ENeuralBackEnd::UEOnly for this platform."));
+	UE_LOG(LogNeuralNetworkInference, Warning, TEXT("FImplBackEndUEAndORT::Run(): Platform or Operating System not suported yet for UEAndORT BackEnd. Set BackEnd to ENeuralBackEnd::Auto or ENeuralBackEnd::UEOnly for this platform."));
 #endif //WITH_UE_AND_ORT_SUPPORT
 }
 
@@ -304,7 +304,7 @@ bool UNeuralNetwork::FImplBackEndUEAndORT::InitializedAndConfigureMembers(TShare
 	// Configure InOutImplBackEndUEAndORT
 	if (!InOutImplBackEndUEAndORT->ConfigureMembers(InDeviceType))
 	{
-		UE_LOG(LogNeuralNetworkInference, Warning, TEXT("UNeuralNetwork::InitializedAndConfigureMembers(): ConfigureMembers failed."));
+		UE_LOG(LogNeuralNetworkInference, Warning, TEXT("FImplBackEndUEAndORT::InitializedAndConfigureMembers(): ConfigureMembers failed."));
 		return false;
 	}
 
@@ -332,13 +332,13 @@ bool UNeuralNetwork::FImplBackEndUEAndORT::ConfigureMembers(const ENeuralDeviceT
 		// @todo for Nebojša.Dragosavac & Ginés.Hidalgo: Handle the default DX11 nicely
 		if (RHIName != TEXT("D3D12"))
 		{
-			//UE_LOG(LogNeuralNetworkInference, Warning, TEXT("UNeuralNetwork::ConfigureMembers(): Only DX12 (\"D3D12\")) rendering is compatible with UEAndORT-based NNI. Instead, \"%s\" was used."), *RHIName);
+			//UE_LOG(LogNeuralNetworkInference, Warning, TEXT("FImplBackEndUEAndORT::ConfigureMembers(): Only DX12 (\"D3D12\")) rendering is compatible with UEAndORT-based NNI. Instead, \"%s\" was used."), *RHIName);
 			//return false;
 
 			// Temporary workaround to avoid breaking NNI
 			if (OrtSessionOptionsAppendExecutionProvider_DML(*SessionOptions, 0))
 			{
-				UE_LOG(LogNeuralNetworkInference, Warning, TEXT("UNeuralNetwork::ConfigureMembers(): Some error occurred when using OrtSessionOptionsAppendExecutionProvider_DML()."));
+				UE_LOG(LogNeuralNetworkInference, Warning, TEXT("FImplBackEndUEAndORT::ConfigureMembers(): Some error occurred when using OrtSessionOptionsAppendExecutionProvider_DML()."));
 				return false;
 			}
 			return true;
@@ -350,7 +350,7 @@ bool UNeuralNetwork::FImplBackEndUEAndORT::ConfigureMembers(const ENeuralDeviceT
 
 		if (Rhi->GetNumAdapters() > 1 || Rhi->GetAdapter(0).GetDesc().NumDeviceNodes > 1)
 		{
-			UE_LOG(LogNeuralNetworkInference, Warning, TEXT("UNeuralNetwork::ConfigureMembers(): There are multiple (%d) adapters and/or multiple (%d) devices, using device at index 0."),
+			UE_LOG(LogNeuralNetworkInference, Warning, TEXT("FImplBackEndUEAndORT::ConfigureMembers(): There are multiple (%d) adapters and/or multiple (%d) devices, using device at index 0."),
 				Rhi->GetNumAdapters(), Rhi->GetAdapter(0).GetDesc().NumDeviceNodes);
 			return false;
 		}
@@ -362,7 +362,7 @@ bool UNeuralNetwork::FImplBackEndUEAndORT::ConfigureMembers(const ENeuralDeviceT
 
 		if (!DmlDevice)
 		{
-			UE_LOG(LogNeuralNetworkInference, Warning, TEXT("UNeuralNetwork::ConfigureMembers(): Invalid DML device found."));
+			UE_LOG(LogNeuralNetworkInference, Warning, TEXT("FImplBackEndUEAndORT::ConfigureMembers(): Invalid DML device found."));
 			return false;
 		}
 
@@ -374,17 +374,17 @@ bool UNeuralNetwork::FImplBackEndUEAndORT::ConfigureMembers(const ENeuralDeviceT
 		SessionOptions->SetGraphOptimizationLevel(GraphOptimizationLevel::ORT_ENABLE_ALL); // ORT_ENABLE_ALL, ORT_ENABLE_EXTENDED, ORT_ENABLE_BASIC, ORT_DISABLE_ALL
 		if (OrtSessionOptionsAppendExecutionProviderEx_DML(*SessionOptions, DmlDevice, NativeCmdQ))
 		{
-			UE_LOG(LogNeuralNetworkInference, Warning, TEXT("UNeuralNetwork::ConfigureMembers(): Some error occurred when using OrtSessionOptionsAppendExecutionProviderEx_DML()."));
+			UE_LOG(LogNeuralNetworkInference, Warning, TEXT("FImplBackEndUEAndORT::ConfigureMembers(): Some error occurred when using OrtSessionOptionsAppendExecutionProviderEx_DML()."));
 			return false;
 		}
 		return true; // @todo: Remove this line when NNI_HLSL is working
 #else
-		UE_LOG(LogNeuralNetworkInference, Warning, TEXT("UNeuralNetwork: GPU mode only supported in Windows for now. Please, switch to CPU or to Windows."));
+		UE_LOG(LogNeuralNetworkInference, Warning, TEXT("FImplBackEndUEAndORT::ConfigureMembers(): GPU mode only supported in Windows for now. Please, switch to CPU or to Windows."));
 
 		//SessionOptions->SetGraphOptimizationLevel(GraphOptimizationLevel::ORT_DISABLE_ALL); // ORT_ENABLE_ALL, ORT_ENABLE_EXTENDED, ORT_ENABLE_BASIC, ORT_DISABLE_ALL
 		//if (OrtSessionOptionsAppendExecutionProvider_NNI_HLSL(*SessionOptions, 0))
 		//{
-		//	UE_LOG(LogNeuralNetworkInference, Warning, TEXT("Some error occurred."));
+		//	UE_LOG(LogNeuralNetworkInference, Warning, TEXT("FImplBackEndUEAndORT::ConfigureMembers(): Some error occurred."));
 		//	return false;
 		//}
 #endif //PLATFORM_WIN64
@@ -397,7 +397,7 @@ bool UNeuralNetwork::FImplBackEndUEAndORT::ConfigureMembers(const ENeuralDeviceT
 		SessionOptions->SetGraphOptimizationLevel(GraphOptimizationLevel::ORT_DISABLE_ALL); // ORT_ENABLE_ALL, ORT_ENABLE_EXTENDED, ORT_ENABLE_BASIC, ORT_DISABLE_ALL
 		if (OrtSessionOptionsAppendExecutionProvider_NNI_CPU(*SessionOptions))
 		{
-			UE_LOG(LogNeuralNetworkInference, Warning, TEXT("Some error occurred."));
+			UE_LOG(LogNeuralNetworkInference, Warning, TEXT("FImplBackEndUEAndORT::ConfigureMembers(): OrtSessionOptionsAppendExecutionProvider_NNI_CPU failed."));
 			return false;
 		}
 #else
@@ -434,69 +434,64 @@ void UNeuralNetwork::FImplBackEndUEAndORT::ConfigureTensors(TArray<FNeuralTensor
 
 		Ort::TensorTypeAndShapeInfo CurrentTensorInfo = CurrentTypeInfo.GetTensorTypeAndShapeInfo();
 
+		ENeuralDataType TensorDataType;
 		{
-			ENeuralDataType TensorDataType;
+			const ONNXTensorElementDataType ONNXTensorElementDataTypeEnum = CurrentTensorInfo.GetElementType();
+			if (ONNXTensorElementDataTypeEnum == ONNX_TENSOR_ELEMENT_DATA_TYPE_FLOAT)
 			{
-				const ONNXTensorElementDataType ONNXTensorElementDataTypeEnum = CurrentTensorInfo.GetElementType();
-				if (ONNXTensorElementDataTypeEnum == ONNX_TENSOR_ELEMENT_DATA_TYPE_FLOAT)
-				{
-					TensorDataType = ENeuralDataType::Float;
-				}
-				//else if (ONNXTensorElementDataTypeEnum == ONNX_TENSOR_ELEMENT_DATA_TYPE_INT32)
-				//{
-				//	TensorDataType = ENeuralDataType::Int32;
-				//}
-				//else if (ONNXTensorElementDataTypeEnum == ONNX_TENSOR_ELEMENT_DATA_TYPE_INT64)
-				//{
-				//	TensorDataType = ENeuralDataType::Int64;
-				//}
-				//else if (ONNXTensorElementDataTypeEnum == ONNX_TENSOR_ELEMENT_DATA_TYPE_UINT32)
-				//{
-				//	TensorDataType = ENeuralDataType::UInt32;
-				//}
-				//else if (ONNXTensorElementDataTypeEnum == ONNX_TENSOR_ELEMENT_DATA_TYPE_UINT64)
-				//{
-				//	TensorDataType = ENeuralDataType::UInt64;
-				//}
-				else
-				{
-					TensorDataType = ENeuralDataType::None;
-					UE_LOG(LogNeuralNetworkInference, Warning, TEXT("ONNXTensorElementDataTypeEnum = %d not implemented yet."), (int32)ONNXTensorElementDataTypeEnum);
-					return;
-				}
+				TensorDataType = ENeuralDataType::Float;
+			}
+			//else if (ONNXTensorElementDataTypeEnum == ONNX_TENSOR_ELEMENT_DATA_TYPE_INT32)
+			//{
+			//	TensorDataType = ENeuralDataType::Int32;
+			//}
+			//else if (ONNXTensorElementDataTypeEnum == ONNX_TENSOR_ELEMENT_DATA_TYPE_INT64)
+			//{
+			//	TensorDataType = ENeuralDataType::Int64;
+			//}
+			//else if (ONNXTensorElementDataTypeEnum == ONNX_TENSOR_ELEMENT_DATA_TYPE_UINT32)
+			//{
+			//	TensorDataType = ENeuralDataType::UInt32;
+			//}
+			//else if (ONNXTensorElementDataTypeEnum == ONNX_TENSOR_ELEMENT_DATA_TYPE_UINT64)
+			//{
+			//	TensorDataType = ENeuralDataType::UInt64;
+			//}
+			else
+			{
+				TensorDataType = ENeuralDataType::None;
+				UE_LOG(LogNeuralNetworkInference, Warning, TEXT("FImplBackEndUEAndORT::ConfigureTensors(): ONNXTensorElementDataTypeEnum = %d not implemented yet."), (int32)ONNXTensorElementDataTypeEnum);
+				return;
 			}
-			TensorDataTypes.Push(TensorDataType);
 		}
+		TensorDataTypes.Push(TensorDataType);
 
 		// Get input shapes/dims
+		TArray<int64> CurrentTensorSizes;
 		{
-			TArray<int64> CurrentTensorSizes;
+			for (const int64_t CurrentTensorSize : CurrentTensorInfo.GetShape())
 			{
-				for (const int64_t CurrentTensorSize : CurrentTensorInfo.GetShape())
+				if (OutAreInputTensorSizesVariable)
 				{
-					if (OutAreInputTensorSizesVariable)
-					{
-						(*OutAreInputTensorSizesVariable)[TensorIndex] |= (CurrentTensorSize < 0);
-					}
-					// Negative (variable) dimensions not implemented yet
-					if (CurrentTensorSize < 0)
-					{
-						CurrentTensorSizes.Push(1);
-						UE_LOG(LogNeuralNetworkInference, Display,
-							TEXT("Negative (i.e., variable) dimensions not allowed yet, hard-coded to 1. Let us know if you really need variable dimensions."
-								" Keep in mind that fixed sizes might allow additional optimizations and speedup of the network during Run()."));
-					}
-					else
-					{
-						CurrentTensorSizes.Push(CurrentTensorSize);
-					}
+					(*OutAreInputTensorSizesVariable)[TensorIndex] |= (CurrentTensorSize < 0);
+				}
+				// Negative (variable) dimensions not implemented yet
+				if (CurrentTensorSize < 0)
+				{
+					CurrentTensorSizes.Push(1);
+					UE_LOG(LogNeuralNetworkInference, Display,
+						TEXT("Negative (i.e., variable) dimensions not allowed yet, hard-coded to 1. Let us know if you really need variable dimensions."
+							" Keep in mind that fixed sizes might allow additional optimizations and speedup of the network during Run()."));
+				}
+				else
+				{
+					CurrentTensorSizes.Push(CurrentTensorSize);
 				}
 			}
-			TensorSizes.Push(CurrentTensorSizes);
 		}
+		TensorSizes.Push(CurrentTensorSizes);
 
 		CurrentTypeInfo.release();
-
 	}
 
 	SetTensorsFromNetwork(OutTensors, TensorNames, TensorDataTypes, TensorSizes, bIsInput);
@@ -508,7 +503,7 @@ void UNeuralNetwork::FImplBackEndUEAndORT::SetTensorsFromNetwork(TArray<FNeuralT
 	const int32 TensorNumber = InTensorNames.Num();
 	if (InTensorDataTypes.Num() != TensorNumber || InSizes.Num() != TensorNumber)
 	{
-		UE_LOG(LogNeuralNetworkInference, Warning, TEXT("InTensorNames.Num() == InTensorDataTypes.Num() == InSizes.Num() failed, %d vs. %d vs. %d."),
+		UE_LOG(LogNeuralNetworkInference, Warning, TEXT("FImplBackEndUEAndORT::SetTensorsFromNetwork(): InTensorNames.Num() == InTensorDataTypes.Num() == InSizes.Num() failed, %d vs. %d vs. %d."),
 			InTensorNames.Num(), InTensorDataTypes.Num(), InSizes.Num());
 		return;
 	}
@@ -589,7 +584,7 @@ void UNeuralNetwork::FImplBackEndUEAndORT::LinkTensorToONNXRuntime(TArray<FNeura
 		//}
 		else
 		{
-			UE_LOG(LogNeuralNetworkInference, Warning, TEXT("Not implemented (yet) for ENeuralDataType = %d."), (int32)NeuralDataType);
+			UE_LOG(LogNeuralNetworkInference, Warning, TEXT("FImplBackEndUEAndORT::LinkTensorToONNXRuntime(): Not implemented (yet) for ENeuralDataType = %d."), (int32)NeuralDataType);
 		}
 	}
 }
diff --git a/Engine/Plugins/Experimental/NeuralNetworkInference/Source/NeuralNetworkInference/Private/NeuralNetworkImplBackEndUEOnly.cpp b/Engine/Plugins/Experimental/NeuralNetworkInference/Source/NeuralNetworkInference/Private/NeuralNetworkImplBackEndUEOnly.cpp
--- a/Engine/Plugins/Experimental/NeuralNetworkInference/Source/NeuralNetworkInference/Private/NeuralNetworkImplBackEndUEOnly.cpp
+++ b/Engine/Plugins/Experimental/NeuralNetworkInference/Source/NeuralNetworkInference/Private/NeuralNetworkImplBackEndUEOnly.cpp
@@ -13,7 +13,7 @@
 /* FImplBackEndUEOnly public functions
  *****************************************************************************/
 
-bool UNeuralNetwork::FImplBackEndUEOnly::Load(TSharedPtr<FImplBackEndUEOnly>& InOutImplBackEndUEOnly, const TArray<uint8>& InModelReadFromFileInBytes, const FString& InModelFilePath)
+bool UNeuralNetwork::FImplBackEndUEOnly::Load(TSharedPtr<FImplBackEndUEOnly>& InOutImplBackEndUEOnly, const TArray<uint8>& InModelReadFromFileInBytes)
 {
 	// Initialize InOutImplBackEndUEOnly
 	if (!InOutImplBackEndUEOnly.IsValid())
@@ -21,19 +21,32 @@ bool UNeuralNetwork::FImplBackEndUEOnly::Load(TSharedPtr<FImplBackEndUEOnly>& In
 		InOutImplBackEndUEOnly = MakeShared<FImplBackEndUEOnly>();
 	}
 	// Clean previous networks
-	InOutImplBackEndUEOnly->Operators.Empty();
-	InOutImplBackEndUEOnly->ModelProto = FModelProto();
-	InOutImplBackEndUEOnly->bAreTensorsInGpu = false;
+	else
+	{
+		InOutImplBackEndUEOnly->Operators.Empty();
+		InOutImplBackEndUEOnly->ModelProto = FModelProto();
+		InOutImplBackEndUEOnly->bAreTensorsInGpu = false;
+	}
 	// Read ModelProto
-	if (!FModelProtoFileReader::ReadModelProtoFromFile(InOutImplBackEndUEOnly->ModelProto, InModelFilePath) || !InOutImplBackEndUEOnly->ModelProto.IsLoaded())
+	if (!FModelProtoFileReader::ReadModelProtoFromArray(InOutImplBackEndUEOnly->ModelProto, InModelReadFromFileInBytes))
 	{
-		UE_LOG(LogNeuralNetworkInference, Warning, TEXT("UNeuralNetwork::FImplBackEndUEOnly::Load(): Model could not be loaded from %s."), *InModelFilePath);
+		UE_LOG(LogNeuralNetworkInference, Warning, TEXT("FImplBackEndUEOnly::Load(): ReadModelProtoFromFile failed."));
+		return false;
+	}
+	if (!InOutImplBackEndUEOnly->ModelProto.IsLoaded())
+	{
+		UE_LOG(LogNeuralNetworkInference, Warning, TEXT("FImplBackEndUEOnly::Load(): ModelProto.IsLoaded() returned false."));
 		return false;
 	}
 	// Turn ModelProto into Operators
-	if (!FGraphProtoToNeuralNetworkConverter::Translate(InOutImplBackEndUEOnly->Operators, InOutImplBackEndUEOnly->TensorManager, InOutImplBackEndUEOnly->ModelProto.GetGraph(), /*bInIsTensorManagerConst*/false) || !InOutImplBackEndUEOnly->TensorManager.IsLoaded())
+	if (!FGraphProtoToNeuralNetworkConverter::Translate(InOutImplBackEndUEOnly->Operators, InOutImplBackEndUEOnly->TensorManager, InOutImplBackEndUEOnly->ModelProto.GetGraph(), /*bInIsTensorManagerConst*/false))
+	{
+		UE_LOG(LogNeuralNetworkInference, Warning, TEXT("FImplBackEndUEOnly::Load(): Translate failed."));
+		return false;
+	}
+	if (!InOutImplBackEndUEOnly->TensorManager.IsLoaded())
 	{
-		UE_LOG(LogNeuralNetworkInference, Warning, TEXT("UNeuralNetwork::FImplBackEndUEOnly::Load(): TensorManager could not be loaded from %s."), *InModelFilePath);
+		UE_LOG(LogNeuralNetworkInference, Warning, TEXT("FImplBackEndUEOnly::Load(): TensorManager.IsLoaded() returned false."));
 		return false;
 	}
 	// Load successful
@@ -50,7 +63,7 @@ bool UNeuralNetwork::FImplBackEndUEOnly::Load(TSharedPtr<FImplBackEndUEOnly>& In
 //	// Load
 //	if (!InTensorManager.IsLoaded())
 //	{
-//		UE_LOG(LogNeuralNetworkInference, Warning, TEXT("UNeuralNetwork::FImplBackEndUEOnly::Load(): TensorManager could not be loaded."));
+//		UE_LOG(LogNeuralNetworkInference, Warning, TEXT("FImplBackEndUEOnly::Load(): TensorManager could not be loaded."));
 //	}
 //	Swap(InOutImplBackEndUEOnly->TensorManager, InTensorManager);
 //	InOutImplBackEndUEOnly->Operators = InOperators;
@@ -59,7 +72,7 @@ bool UNeuralNetwork::FImplBackEndUEOnly::Load(TSharedPtr<FImplBackEndUEOnly>& In
 
 void UNeuralNetwork::FImplBackEndUEOnly::Run(FOnAsyncRunCompleted& InOutOnAsyncRunCompletedDelegate, const ENeuralNetworkSynchronousMode InSynchronousMode, const ENeuralDeviceType InDeviceType, const ENeuralDeviceType InInputDeviceType, const ENeuralDeviceType InOutputDeviceType)
 {
-	// Run UNeuralNetworkLegacy
+	// Run UNeuralNetwork::UEOnly
 	if (Operators.Num() > 0)
 	{
 		// Run graph
@@ -83,7 +96,7 @@ void UNeuralNetwork::FImplBackEndUEOnly::Run(FOnAsyncRunCompleted& InOutOnAsyncR
 			// Sanity check
 			if (TensorManager.GetTensorsMutable().Num() < 1)
 			{
-				UE_LOG(LogNeuralNetworkInference, Warning, TEXT("UNeuralNetwork::FImplBackEndUEOnly::Run(): Tensors.Num() = %d (should be > 0)."), TensorManager.GetTensorsMutable().Num());
+				UE_LOG(LogNeuralNetworkInference, Warning, TEXT("FImplBackEndUEOnly::Run(): Tensors.Num() = %d (should be > 0)."), TensorManager.GetTensorsMutable().Num());
 				return;
 			}
 
@@ -91,12 +104,12 @@ void UNeuralNetwork::FImplBackEndUEOnly::Run(FOnAsyncRunCompleted& InOutOnAsyncR
 			ENQUEUE_RENDER_COMMAND(UNeuralNetwork_UEOnly_Run_RenderThread)(
 				[this, &InOutOnAsyncRunCompletedDelegate, InSynchronousMode, InInputDeviceType, InOutputDeviceType](FRHICommandListImmediate& RHICmdList)
 				{
-					FRDGBuilder GraphBuilder(RHICmdList, RDG_EVENT_NAME("UNeuralNetwork::FImplBackEndUEOnly::Run()"));
+					FRDGBuilder GraphBuilder(RHICmdList, RDG_EVENT_NAME("FImplBackEndUEOnly::Run()"));
 
 					// Move memory from CPU to GPU
 					TArray<FNeuralTensor>& Tensors = TensorManager.GetTensorsMutable();
 					const bool bIsInputInCPU = (InInputDeviceType == ENeuralDeviceType::CPU);
-					// Move only input tensors to GPU (once per UNeuralNetwork::FImplBackEndUEOnly::Run())
+					// Move only input tensors to GPU (once per FImplBackEndUEOnly::Run())
 					if (bAreTensorsInGpu)
 					{
 						// If input is on CPU, move inputs to GPU
@@ -121,7 +134,7 @@ void UNeuralNetwork::FImplBackEndUEOnly::Run(FOnAsyncRunCompleted& InOutOnAsyncR
 						}
 					}
 					// Move all (input, intermediate, weight, output) tensors to GPU and also call Operators.ToGPU() to move their auxiliary memory into GPU
-					// Only once per UNeuralNetworkLegacy instance, or until Load() is called again
+					// Only once per UNeuralNetwork instance, or until Load() is called again
 					else
 					{
 						// If input is on CPU, move inputs to GPU
@@ -201,12 +214,12 @@ void UNeuralNetwork::FImplBackEndUEOnly::Run(FOnAsyncRunCompleted& InOutOnAsyncR
 		}
 		else
 		{
-			UE_LOG(LogNeuralNetworkInference, Warning, TEXT("UNeuralNetwork::FImplBackEndUEOnly::Run(): Unknown DeviceType = %d."), (int32)InDeviceType);
+			UE_LOG(LogNeuralNetworkInference, Warning, TEXT("FImplBackEndUEOnly::Run(): Unknown DeviceType = %d."), (int32)InDeviceType);
 		}
 	}
 	else
 	{
-		UE_LOG(LogNeuralNetworkInference, Warning, TEXT("UNeuralNetwork::FImplBackEndUEOnly::Run() called with an empty model."));
+		UE_LOG(LogNeuralNetworkInference, Warning, TEXT("FImplBackEndUEOnly::Run() called with an empty model."));
 	}
 }
 
diff --git a/Engine/Plugins/Experimental/NeuralNetworkInference/Source/NeuralNetworkInference/Private/NeuralNetworkImplBackEndUEOnly.h b/Engine/Plugins/Experimental/NeuralNetworkInference/Source/NeuralNetworkInference/Private/NeuralNetworkImplBackEndUEOnly.h
--- a/Engine/Plugins/Experimental/NeuralNetworkInference/Source/NeuralNetworkInference/Private/NeuralNetworkImplBackEndUEOnly.h
+++ b/Engine/Plugins/Experimental/NeuralNetworkInference/Source/NeuralNetworkInference/Private/NeuralNetworkImplBackEndUEOnly.h
@@ -31,7 +31,7 @@ struct UNeuralNetwork::FImplBackEndUEOnly
 	 */
 	TArray<TSharedPtr<FNeuralOperator>> Operators;
 
- 	static bool Load(TSharedPtr<FImplBackEndUEOnly>& InOutImplBackEndUEOnly, const TArray<uint8>& InModelReadFromFileInBytes, const FString& InModelFullFilePath);
+ 	static bool Load(TSharedPtr<FImplBackEndUEOnly>& InOutImplBackEndUEOnly, const TArray<uint8>& InModelReadFromFileInBytes);
 
 	//static bool Load(TSharedPtr<FImplBackEndUEOnly>& InOutImplBackEndUEOnly, FNeuralTensorManager& InTensorManager, const TArray<TSharedPtr<FNeuralOperator>>& InOperators);
 
