NeuralNetworkInference (NNI) plugin:
- Mini bug fixed: Removed WITH_UE_AND_ORT_SUPPORT from NeuralNetworkInferenceEditorModule.

#rnx

diff --git a/Engine/Plugins/Experimental/NeuralNetworkInference/Source/NeuralNetworkInferenceEditor/Private/NeuralNetworkInferenceEditorModule.cpp b/Engine/Plugins/Experimental/NeuralNetworkInference/Source/NeuralNetworkInferenceEditor/Private/NeuralNetworkInferenceEditorModule.cpp
--- a/Engine/Plugins/Experimental/NeuralNetworkInference/Source/NeuralNetworkInferenceEditor/Private/NeuralNetworkInferenceEditorModule.cpp
+++ b/Engine/Plugins/Experimental/NeuralNetworkInference/Source/NeuralNetworkInferenceEditor/Private/NeuralNetworkInferenceEditorModule.cpp
@@ -36,10 +36,8 @@ void FNeuralNetworkInferenceEditorModule::StartupModule()
 {
 	// UNeuralNetwork - Register asset types
 	IAssetTools& AssetTools = FModuleManager::LoadModuleChecked<FAssetToolsModule>("AssetTools").Get();
-#ifdef WITH_UE_AND_ORT_SUPPORT
 	Actions.Emplace(MakeShared<FNeuralNetworkAssetTypeActions>());
 	AssetTools.RegisterAssetTypeActions(Actions.Last().ToSharedRef());
-#endif
 	Actions.Emplace(MakeShared<FNeuralNetworkLegacyAssetTypeActions>());
 	AssetTools.RegisterAssetTypeActions(Actions.Last().ToSharedRef());
 	// Register ML category so that ML assets can register to it
