NeuralNetworkInference (NNI) plugin:
- Bug fixed: FGraphProtoToNeuralNetworkConverter::CreateOperatorsAndEditTensorArray() will return with a warning (rather than crashing) when operator not implemented.

#rnx

diff --git a/Engine/Plugins/Experimental/NeuralNetworkInference/Source/NeuralNetworkInference/Private/GraphProtoToNeuralNetworkConverter.cpp b/Engine/Plugins/Experimental/NeuralNetworkInference/Source/NeuralNetworkInference/Private/GraphProtoToNeuralNetworkConverter.cpp
--- a/Engine/Plugins/Experimental/NeuralNetworkInference/Source/NeuralNetworkInference/Private/GraphProtoToNeuralNetworkConverter.cpp
+++ b/Engine/Plugins/Experimental/NeuralNetworkInference/Source/NeuralNetworkInference/Private/GraphProtoToNeuralNetworkConverter.cpp
@@ -104,11 +104,6 @@ bool FGraphProtoToNeuralNetworkConverter::Translate(TArray<TSharedPtr<FNeuralOpe
 		// Fill operators
 		TMap<FString, int32> OutputNameIndexMap;
 		bool bWasSuccessful = CreateOperatorsAndEditTensorArray(OutOperators, OutputNameIndexMap, Tensors, NameIndexMap, InputNameIndexMap, OutputNameDummyIndexMap, InGraphProto, bInIsTensorManagerConst);
-		if (Tensors.Num() <= NameIndexMap.Num())
-		{
-			UE_LOG(LogNeuralNetworkInference, Display, TEXT("Auto-inlining: %d initial tensors compressed into %d (%d tensors merged)."),
-				NameIndexMap.Num(), Tensors.Num(), NameIndexMap.Num() - Tensors.Num());
-		}
 		// Sanity check
 		if (Tensors.Num() > MaxNumberTensorsInNetwork)
 		{
@@ -118,9 +113,14 @@ bool FGraphProtoToNeuralNetworkConverter::Translate(TArray<TSharedPtr<FNeuralOpe
 			bWasSuccessful = false;
 		}
 
-		// Set InOrOutTensorManager
 		if (bWasSuccessful)
 		{
+			if (Tensors.Num() <= NameIndexMap.Num())
+			{
+				UE_LOG(LogNeuralNetworkInference, Display, TEXT("Auto-inlining: %d initial tensors compressed into %d (%d tensors merged)."),
+					NameIndexMap.Num(), Tensors.Num(), NameIndexMap.Num() - Tensors.Num());
+			}
+			// Set InOrOutTensorManager
 			return InOrOutTensorManager.Load(Tensors, NameIndexMap, InputNameIndexMap, OutputNameIndexMap);
 		}
 		return bWasSuccessful;
@@ -229,7 +229,7 @@ bool FGraphProtoToNeuralNetworkConverter::CreateOperatorsAndEditTensorArray(TArr
 				// Create constant operator layer and get final tensor value
 				FConstantOperator ConstantOperator(NodeProto);
 				ConstantOperator.SetOutputTensors({&InOutTensors.Last()});
-				// Sanity checks
+				// Sanity check
 				if (!ConstantOperator.ConfigureOutputAndInternalVariablesAndSanityChecks())
 				{
 					UE_LOG(LogNeuralNetworkInference, Warning, TEXT("FGraphProtoToNeuralNetworkConverter::CreateOperatorsAndEditTensorArray()-Constant: Tensor %s was already defined and it should have not."), *OutputTensorName);
@@ -244,6 +244,10 @@ bool FGraphProtoToNeuralNetworkConverter::CreateOperatorsAndEditTensorArray(TArr
 		{
 			// Create and configure operator
 			TSharedPtr<FNeuralOperator> Operator = CreateOperator(NodeProto, OperatorInputTensorNames, InInputNameDummyIndexMap, InOutputNameDummyIndexMap, InGraphProto);
+			if (!Operator.IsValid())
+			{
+				return false;
+			}
 			// ConvTranspose requires flipping OperatorInputTensors[1], so we expect for now a Weight type, so it can be flipped beforehand
 			if (!bInIsTensorManagerConst && NodeProto.OperatorType == TEXT("ConvTranspose")) //if (NodeProto.OperatorType.Contains(TEXT("Conv"), ESearchCase::CaseSensitive))
 			{
