NeuralNetworkInference (NNI) plugin:
Fix missed #ifdef in thirdparywarningdisabler
#rnx

diff --git a/Engine/Plugins/Experimental/NeuralNetworkInference/Source/Internal/ThirdPartyHelperAndDLLLoader/Internal/ThirdPartyWarningDisabler.h b/Engine/Plugins/Experimental/NeuralNetworkInference/Source/Internal/ThirdPartyHelperAndDLLLoader/Internal/ThirdPartyWarningDisabler.h
--- a/Engine/Plugins/Experimental/NeuralNetworkInference/Source/Internal/ThirdPartyHelperAndDLLLoader/Internal/ThirdPartyWarningDisabler.h
+++ b/Engine/Plugins/Experimental/NeuralNetworkInference/Source/Internal/ThirdPartyHelperAndDLLLoader/Internal/ThirdPartyWarningDisabler.h
@@ -14,7 +14,7 @@ Many third party headers require some care when importing. NNI third party inclu
 	NNI_THIRD_PARTY_INCLUDES_END
 */
 
-#if PLATFORM_NNI_MICROSOFT
+#ifdef PLATFORM_NNI_MICROSOFT
 #define NNI_THIRD_PARTY_INCLUDES_START THIRD_PARTY_INCLUDES_START \
 	__pragma(warning(disable: 4191)) /* For ONNX Runtime */ \
 	__pragma(warning(disable: 4471)) /* For ONNX Runtime */ \
