NeuralNetworkInference (NNI) plugin:
- FModelUnitTester::GlobalTest() runs UEOnly and UEAndORT tests (only in editor to avoid linking error).
- CPU & GPU run in different functions, so eventual Neb's code can be added more easily.
- Other cleanup of that function.

#rnx

diff --git a/Engine/Plugins/Experimental/NeuralNetworkInference/Source/NeuralNetworkInference/Private/NeuralNetworkImplBackEndUEAndORT.cpp b/Engine/Plugins/Experimental/NeuralNetworkInference/Source/NeuralNetworkInference/Private/NeuralNetworkImplBackEndUEAndORT.cpp
--- a/Engine/Plugins/Experimental/NeuralNetworkInference/Source/NeuralNetworkInference/Private/NeuralNetworkImplBackEndUEAndORT.cpp
+++ b/Engine/Plugins/Experimental/NeuralNetworkInference/Source/NeuralNetworkInference/Private/NeuralNetworkImplBackEndUEAndORT.cpp
@@ -210,12 +210,9 @@ bool UNeuralNetwork::FImplBackEndUEAndORT::Load(TSharedPtr<FImplBackEndUEAndORT>
 
 #ifdef PLATFORM_WIN64
 			// Check if resource allocator is properly initialized
-			if (InOutImplBackEndUEAndORT->DmlGPUAllocator.IsValid())
+			if (InOutImplBackEndUEAndORT->DmlGPUAllocator.IsValid() && !InOutImplBackEndUEAndORT->DmlGPUAllocator->IsValid())
 			{
-				if (!InOutImplBackEndUEAndORT->DmlGPUAllocator->IsValid())
-				{
-					UE_LOG(LogNeuralNetworkInference, Warning, TEXT("UNeuralNetwork::Load() DirectML GPU resource allocator has failed to initialize"));
-				}
+				UE_LOG(LogNeuralNetworkInference, Warning, TEXT("FImplBackEndUEAndORT::Load() DirectML GPU resource allocator has failed to initialize."));
 			}
 #endif
 		}
@@ -247,13 +244,11 @@ bool UNeuralNetwork::FImplBackEndUEAndORT::Load(TSharedPtr<FImplBackEndUEAndORT>
 
 void UNeuralNetwork::FImplBackEndUEAndORT::ClearResources()
 {
-#ifdef PLATFORM_WIN64
 #ifdef WITH_UE_AND_ORT_SUPPORT
-	
+#ifdef PLATFORM_WIN64
 	DmlGPUAllocator.Reset(nullptr);
-
-#endif
-#endif
+#endif //PLATFORM_WIN64
+#endif //WITH_UE_AND_ORT_SUPPORT
 }
 
 void UNeuralNetwork::FImplBackEndUEAndORT::Run(const ENeuralNetworkSynchronousMode InSynchronousMode, const ENeuralDeviceType InInputDeviceType, const ENeuralDeviceType InOutputDeviceType)
@@ -355,7 +350,7 @@ bool UNeuralNetwork::FImplBackEndUEAndORT::ConfigureMembers(const ENeuralDeviceT
 		// @todo for Nebojša.Dragosavac & Ginés.Hidalgo: Handle the default DX11 nicely
 		if (RHIName != TEXT("D3D12"))
 		{
-			//UE_LOG(LogNeuralNetworkInference, Warning, TEXT("FImplBackEndUEAndORT::ConfigureMembers(): Only DX12 (\"D3D12\")) rendering is compatible with UEAndORT-based NNI. Instead, \"%s\" was used."), *RHIName);
+			//UE_LOG(LogNeuralNetworkInference, Warning, TEXT("FImplBackEndUEAndORT::ConfigureMembers(): Only DX12 (\"D3D12\") rendering is compatible with UEAndORT-based NNI. Instead, \"%s\" was used."), *RHIName);
 			//return false;
 
 			// Temporary workaround to avoid breaking NNI
@@ -366,6 +361,7 @@ bool UNeuralNetwork::FImplBackEndUEAndORT::ConfigureMembers(const ENeuralDeviceT
 			}
 			return true;
 		}
+UE_LOG(LogNeuralNetworkInference, Display, TEXT("Testing the experimental DX12 (\"D3D12\") back end."));
 
 		// Get adapter's D3D12 device that we would like to share with DirectML execution provider
 		// NOTE: For now we're only using first device that has Dadapter 0 and device 0
diff --git a/Engine/Plugins/Experimental/NeuralNetworkInference/Source/NeuralNetworkInferenceQA/Private/ModelUnitTester.cpp b/Engine/Plugins/Experimental/NeuralNetworkInference/Source/NeuralNetworkInferenceQA/Private/ModelUnitTester.cpp
--- a/Engine/Plugins/Experimental/NeuralNetworkInference/Source/NeuralNetworkInferenceQA/Private/ModelUnitTester.cpp
+++ b/Engine/Plugins/Experimental/NeuralNetworkInference/Source/NeuralNetworkInferenceQA/Private/ModelUnitTester.cpp
@@ -13,29 +13,30 @@
 bool FModelUnitTester::GlobalTest(const FString& InProjectContentDir, const FString& InModelZooRelativeDirectory)
 {
 	// Model names, input values, and number of repetitions for profiling
-#ifdef WITH_UE_AND_ORT_SUPPORT
 	const TArray<FString> ModelNames({ TEXT("MLRigDeformer"), TEXT("cloth_network"), TEXT("HS"), TEXT("RL") });
 	const TArray<float> InputArrayValues({ 1.f, 0.f, -1.f, 100.f, -100.f, 0.5f, -0.5f }); // This one can be shorter than CPU/GPUGroundTruths
-	const TArray<int32> CPURepetitions({ 1000, 1000,  50, 1000 });
-	const TArray<int32> GPURepetitions({ 1000, 1000, 100, 1000 });
-#else
-	UE_LOG(LogNeuralNetworkInferenceQA, Display, TEXT("--------------- Some FModelUnitTester tests skipped (only if WITH_UE_AND_ORT_SUPPORT)."));
-	const TArray<FString> ModelNames({ TEXT("MLRigDeformer"), TEXT("cloth_network") });
-	const TArray<float> InputArrayValues({ 1.f, 0.f });
-	const TArray<int32> CPURepetitions({ 100, 0 });
-	const TArray<int32> GPURepetitions({ 100, 100 });
-UE_LOG(LogNeuralNetworkInferenceQA, Display, TEXT("--------------- FModelUnitTester test skipped (only if WITH_UE_AND_ORT_SUPPORT)."));
-return true;
-#endif //WITH_UE_AND_ORT_SUPPORT
-
 	// Ground truths
 	const TArray<TArray<double>> CPUGroundTruths({ {3.728547, 0.008774, 4.595651, 212.193216, 742.434561, 4.250668, 4.717748}, {0.042571, 0.023693, 0.015783, 13.100505, 8.050994, 0.028807, 0.016387},
-		{138.372906, 126.753839, 127.287254, 130.316062, 127.303424, 124.800896, 126.546051}, {0.488662, 0.472437, 0.478862, 0.522685, 0.038322, 0.480848, 0.483821}});
+		{138.372906, 126.753839, 127.287254, 130.316062, 127.303424, 124.800896, 126.546051}, {0.488662, 0.472437, 0.478862, 0.522685, 0.038322, 0.480848, 0.483821} });
 	const TArray<TArray<double>> GPUGroundTruths({ {3.728547, 0.008774, 4.595651, 212.193208, 742.434578, 4.250668, 4.717748}, {0.042571, 0.023693, 0.015783, 13.100504, 8.050994, 0.028807, 0.016387},
-		{138.373184, 126.754100, 127.287398, 130.316194, 127.303495, 124.801134, 126.5462530}, {0.488662, 0.472437, 0.478862, 0.522685, 0.038322, 0.480848, 0.483821}});
-
+		{138.373184, 126.754100, 127.287398, 130.316194, 127.303495, 124.801134, 126.5462530}, {0.488662, 0.472437, 0.478862, 0.522685, 0.038322, 0.480848, 0.483821} });
+	// Speed profiling test - 0 repetitions means that test will not be run
+#ifdef WITH_UE_AND_ORT_SUPPORT
+	const TArray<int32> CPURepetitionsForUEAndORT({ 1000, 1000,  50, 1000 });
+	const TArray<int32> GPURepetitionsForUEAndORT({ 1000, 1000, 100, 1000 });
+#else //WITH_UE_AND_ORT_SUPPORT
+	const TArray<int32> CPURepetitionsForUEAndORT({ 0, 0, 0, 0 });
+	const TArray<int32> GPURepetitionsForUEAndORT({ 0, 0, 0, 0 });
+#endif //WITH_UE_AND_ORT_SUPPORT
+#if WITH_EDITOR
+	const TArray<int32> CPURepetitionsForUEOnlyBackEnd({ 10, 0, 0, 0 });
+	const TArray<int32> GPURepetitionsForUEOnlyBackEnd({ 10, 10, 0, 0 });
+#else //WITH_EDITOR
+	const TArray<int32> CPURepetitionsForUEOnlyBackEnd({ 0, 0, 0, 0 });
+	const TArray<int32> GPURepetitionsForUEOnlyBackEnd({ 0, 0, 0, 0 });
+#endif //WITH_EDITOR
 	// Run tests
-	return ModelLoadAccuracyAndSpeedTests(InProjectContentDir, InModelZooRelativeDirectory, ModelNames, InputArrayValues, CPUGroundTruths, GPUGroundTruths, CPURepetitions, GPURepetitions);
+	return ModelLoadAccuracyAndSpeedTests(InProjectContentDir, InModelZooRelativeDirectory, ModelNames, InputArrayValues, CPUGroundTruths, GPUGroundTruths, CPURepetitionsForUEAndORT, GPURepetitionsForUEAndORT, CPURepetitionsForUEOnlyBackEnd, GPURepetitionsForUEOnlyBackEnd);
 }
 
 
@@ -44,8 +45,8 @@ return true;
  *****************************************************************************/
 
 bool FModelUnitTester::ModelLoadAccuracyAndSpeedTests(const FString& InProjectContentDir, const FString& InModelZooRelativeDirectory, const TArray<FString>& InModelNames,
-	const TArray<float>& InInputArrayValues, const TArray<TArray<double>>& InCPUGroundTruths, const TArray<TArray<double>>& InGPUGroundTruths, const TArray<int32>& InCPURepetitions,
-	const TArray<int32>& InGPURepetitions)
+	const TArray<float>& InInputArrayValues, const TArray<TArray<double>>& InCPUGroundTruths, const TArray<TArray<double>>& InGPUGroundTruths, const TArray<int32>& InCPURepetitionsForUEAndORT,
+	const TArray<int32>& InGPURepetitionsForUEAndORT, const TArray<int32>& InCPURepetitionsForUEOnlyBackEnd, const TArray<int32>& InGPURepetitionsForUEOnlyBackEnd)
 {
 	bool bDidGlobalTestPassed = true;
 
@@ -63,7 +64,7 @@ bool FModelUnitTester::ModelLoadAccuracyAndSpeedTests(const FString& InProjectCo
 		UNeuralNetwork* Network = NetworkUassetLoadTest(UAssetModelFilePath);
 		if (!Network)
 		{
-			ensureMsgf(false, TEXT("UNeuralNetwork could not be loaded from UAssetModelFilePath %s."), *UAssetModelFilePath);
+			UE_LOG(LogNeuralNetworkInferenceQA, Warning, TEXT("UNeuralNetwork could not be loaded from UAssetModelFilePath %s."), *UAssetModelFilePath);
 			return false;
 		}
 		// Input debugging
@@ -83,19 +84,15 @@ bool FModelUnitTester::ModelLoadAccuracyAndSpeedTests(const FString& InProjectCo
 
 		UE_LOG(LogNeuralNetworkInferenceQA, Display, TEXT("--------------- %s - Network ONNX/ORT Load and Run"), *ModelName);
 #if WITH_EDITOR
-		for (int32 ONNXOrORTIndex = 0; ONNXOrORTIndex < 2; ++ONNXOrORTIndex)
+		const int32 ONNXOrORTIndex = 0; //for (int32 ONNXOrORTIndex = 0; ONNXOrORTIndex < 2; ++ONNXOrORTIndex)
 		{
 			const FString ModelFilePath = (ONNXOrORTIndex % 2 == 0 ? GetONNXModelFilePath(ModelZooDirectory, ModelName) : GetORTModelFilePath(ModelZooDirectory, ModelName));
 			const FString ModelType = (ONNXOrORTIndex % 2 == 0 ? TEXT("ONNX") : TEXT("ORT"));
-if (ONNXOrORTIndex == 1)
-{
-	break;
-}
 			UE_LOG(LogNeuralNetworkInferenceQA, Display, TEXT("-------------------- %s - Network %s Load and Run - %s"), *ModelName, *ModelType, *ModelFilePath);
 			bDidGlobalTestPassed &= ModelAccuracyTest(NetworkONNXOrORTLoadTest(ModelFilePath), InInputArrayValues, CPUGroundTruths, GPUGroundTruths);
 		}
 #else //WITH_EDITOR
-	UE_LOG(LogNeuralNetworkInferenceQA, Display, TEXT("-------------------- Skipped (only in WITH_EDITOR enabled)."));
+	UE_LOG(LogNeuralNetworkInferenceQA, Display, TEXT("-------------------- Skipped (only if WITH_EDITOR enabled)."));
 #endif //WITH_EDITOR
 	}
 
@@ -105,10 +102,23 @@ if (ONNXOrORTIndex == 1)
 		const FString& ModelName = InModelNames[ModelIndex];
 		UE_LOG(LogNeuralNetworkInferenceQA, Display, TEXT("--------------- %s - Network UAsset Speed Profiling"), *ModelName);
 		const FString UAssetModelFilePath = GetUAssetModelFilePath(ModelName, InModelZooRelativeDirectory);
-		bDidGlobalTestPassed &= ModelSpeedTest(UAssetModelFilePath, InCPURepetitions[ModelIndex], InGPURepetitions[ModelIndex]);
+
+		// UEAndORT (if WITH_UE_AND_ORT_SUPPORT)
+#ifdef WITH_UE_AND_ORT_SUPPORT
+		bDidGlobalTestPassed &= ModelSpeedTest(UAssetModelFilePath, ENeuralDeviceType::CPU, ENeuralBackEnd::UEAndORT, InCPURepetitionsForUEAndORT[ModelIndex]);
+#ifdef PLATFORM_WIN64
+		bDidGlobalTestPassed &= ModelSpeedTest(UAssetModelFilePath, ENeuralDeviceType::GPU, ENeuralBackEnd::UEAndORT, InGPURepetitionsForUEAndORT[ModelIndex]);
+#else //PLATFORM_WIN64
+		bDidGlobalTestPassed &= ModelSpeedTest(UAssetModelFilePath, ENeuralDeviceType::GPU, ENeuralBackEnd::UEAndORT, 0);
+#endif //PLATFORM_WIN64
+#endif //WITH_UE_AND_ORT_SUPPORT
+
+		// UEOnly
+		bDidGlobalTestPassed &= ModelSpeedTest(UAssetModelFilePath, ENeuralDeviceType::CPU, ENeuralBackEnd::UEOnly, InCPURepetitionsForUEOnlyBackEnd[ModelIndex]);
+		bDidGlobalTestPassed &= ModelSpeedTest(UAssetModelFilePath, ENeuralDeviceType::GPU, ENeuralBackEnd::UEOnly, InGPURepetitionsForUEOnlyBackEnd[ModelIndex]);
 	}
 
-	return true;
+	return bDidGlobalTestPassed;
 }
 
 FString FModelUnitTester::GetONNXModelFilePath(const FString& ModelZooDirectory, const FString& InModelName)
@@ -177,9 +187,12 @@ bool FModelUnitTester::ModelAccuracyTest(UNeuralNetwork* InOutNetwork, const TAr
 			InputArrays.Emplace(TArray<float>({}));
 			InputArrays.Last().Init(InputArrayValue, NetworkSize);
 		}
-		ensureMsgf(InputArrays.Num() <= InCPUGroundTruths.Num() && InputArrays.Num() <= InGPUGroundTruths.Num(),
-			TEXT("InputArrays.Num() <= InCPUGroundTruths.Num() && InputArrays.Num() <= InGPUGroundTruths.Num() failed: %d vs. %d vs. %d."),
-			InputArrays.Num(), InCPUGroundTruths.Num(), InGPUGroundTruths.Num());
+		if (InputArrays.Num() > InCPUGroundTruths.Num() || InputArrays.Num() > InGPUGroundTruths.Num())
+		{
+			UE_LOG(LogNeuralNetworkInferenceQA, Warning, TEXT("InputArrays.Num() <= InCPUGroundTruths.Num() && InputArrays.Num() <= InGPUGroundTruths.Num() failed: %d vs. %d vs. %d."),
+				InputArrays.Num(), InCPUGroundTruths.Num(), InGPUGroundTruths.Num());
+			return false;
+		}
 	}
 	// Run each input with CPU/GPU and compare with each other and with the ground truth
 	bool bDidGlobalTestPassed = true;
@@ -226,7 +239,7 @@ bool FModelUnitTester::ModelAccuracyTest(UNeuralNetwork* InOutNetwork, const TAr
 				*FNeuralTensor(InOutNetwork->GetInputTensor().GetArrayCopy<float>(), InputSizes).ToString(MaxNumberElementsToDisplay));
 			UE_LOG(LogNeuralNetworkInferenceQA, Display, TEXT("OutputArrayCPU = %s"), *FNeuralTensor(OutputArrayCPU, OutputSizes).ToString(MaxNumberElementsToDisplay));
 			UE_LOG(LogNeuralNetworkInferenceQA, Display, TEXT("OutputArrayGPU = %s"), *FNeuralTensor(OutputArrayGPU, OutputSizes).ToString(MaxNumberElementsToDisplay));
-			ensureMsgf(false, TEXT("At least 1 of the 4 CPU/GPU tests failed."));
+			UE_LOG(LogNeuralNetworkInferenceQA, Warning, TEXT("At least 1 of the 4 CPU/GPU tests failed."));
 			return false;
 		}
 		bDidGlobalTestPassed &= !bDidSomeTestFailed;
@@ -264,16 +277,54 @@ bool FModelUnitTester::ModelAccuracyTest(UNeuralNetwork* InOutNetwork, const TAr
 	} \
 	const float OutNetworkTimeInMilliSeconds = Timer.Toc() / InRepetitions - OutCopyTimeInMilliSeconds
 
-bool FModelUnitTester::ModelSpeedTest(const FString& InUAssetPath, const int32 InCPURepetitions, const int32 InGPURepetitions)
+bool FModelUnitTester::ModelSpeedTest(const FString& InUAssetPath, const ENeuralDeviceType InDeviceType, const ENeuralBackEnd InBackEnd, const int32 InRepetitions)
 {
+	// Get debug strings
+	FString DeviceTypeString;
+	if (InDeviceType == ENeuralDeviceType::CPU)
+	{
+		DeviceTypeString = TEXT("CPU");
+	}
+	else if (InDeviceType == ENeuralDeviceType::GPU)
+	{
+		DeviceTypeString = TEXT("GPU");
+	}
+	else
+	{
+		UE_LOG(LogNeuralNetworkInferenceQA, Warning, TEXT("FModelUnitTester::ModelSpeedTest(): Unknown DeviceTypeString = %d."), (int32)InDeviceType);
+	}
+	FString BackEndForCurrentPlatformString;
+	if (InBackEnd == ENeuralBackEnd::UEAndORT)
+	{
+		BackEndForCurrentPlatformString = TEXT("UEAndORT");
+	}
+	else if (InBackEnd == ENeuralBackEnd::UEOnly)
+	{
+		BackEndForCurrentPlatformString = TEXT("UEOnly");
+	}
+	else
+	{
+		UE_LOG(LogNeuralNetworkInferenceQA, Warning, TEXT("FModelUnitTester::ModelSpeedTest(): Unknown BackEndForCurrentPlatformString = %d."), (int32)InBackEnd);
+	}
+	// Skip test if reps = 0
+	if (InRepetitions < 1)
+	{
+		UE_LOG(LogNeuralNetworkInferenceQA, Display, TEXT("ModelSpeedTest skipped for uasset \"%s\" for InDeviceType = %s (%d) and InBackEnd = %s (%d)."), *InUAssetPath, *DeviceTypeString, (int32)InDeviceType, *BackEndForCurrentPlatformString, (int32)InBackEnd);
+		return true;
+	}
 	// Load Network
 	UNeuralNetwork* InOutNetwork = NetworkUassetLoadTest(InUAssetPath);
-	// Sanity check
+	// Sanity checks
 	if (!InOutNetwork)
 	{
-		ensureMsgf(false, TEXT("InOutNetwork was a nullptr."));
+		UE_LOG(LogNeuralNetworkInferenceQA, Warning, TEXT("FModelUnitTester::ModelSpeedTest(): InOutNetwork was a nullptr. Path: \"%s\"."), *InUAssetPath);
 		return false;
 	}
+	// Get original values
+	const ENeuralDeviceType OriginalDeviceType = InOutNetwork->GetDeviceType();
+	const ENeuralBackEnd OriginalBackEnd = InOutNetwork->GetBackEnd();
+	// Set desired back end
+	InOutNetwork->SetBackEnd(InBackEnd);
 	// Needed variables
 	const int64 NetworkSize = InOutNetwork->GetInputTensor().Num();
 	TArray<float> InputArray;
@@ -281,21 +332,19 @@ bool FModelUnitTester::ModelSpeedTest(const FString& InUAssetPath, const int32 I
 	TArray<float> OutputArrayGPU;
 	// Speed profiling
 	FNeuralNetworkInferenceQATimer Timer;
-	// CPU
-	InOutNetwork->SetDeviceType(ENeuralDeviceType::CPU);
-	MODEL_UNIT_TESTER_SPEED_TEST(CPUCopyTimer1, CPUNetworkTimer1, 1, /*bIsGPU*/false);
-	InOutNetwork->SetDeviceType(ENeuralDeviceType::CPU);
-	MODEL_UNIT_TESTER_SPEED_TEST(CPUCopyTimer, CPUNetworkTimer, InCPURepetitions, /*bIsGPU*/false);
-	// GPU
-	InOutNetwork->SetDeviceType(ENeuralDeviceType::GPU);
-	MODEL_UNIT_TESTER_SPEED_TEST(GPUCopyTimer1, GPUNetworkTimer1, 1, /*bIsGPU*/true);
-	InOutNetwork->SetDeviceType(ENeuralDeviceType::GPU);
-	MODEL_UNIT_TESTER_SPEED_TEST(GPUCopyTimer, GPUNetworkTimer, InGPURepetitions, /*bIsGPU*/true);
+	// Run profiling 1 time
+	InOutNetwork->SetDeviceType(InDeviceType);
+	MODEL_UNIT_TESTER_SPEED_TEST(CopyTimer1, NetworkTimer1, 1, /*bIsGPU*/false);
+	// Run profiling n times
+	InOutNetwork->SetDeviceType(InDeviceType);
+	MODEL_UNIT_TESTER_SPEED_TEST(CopyTimer, NetworkTimer, InRepetitions, /*bIsGPU*/false);
 	// Display speed times
 	UE_LOG(LogNeuralNetworkInferenceQA, Display,
-		TEXT("Forward pass speed profiling (TensorCopy and Forward()): CPUx1 = %f+%f msec, CPUx%d = %f+%f msec, GPUx1 = %f+%f msec, GPUx%d = %f+%f msec."),
-		CPUCopyTimer1, CPUNetworkTimer1, InCPURepetitions, CPUCopyTimer, CPUNetworkTimer,
-		GPUCopyTimer1, GPUNetworkTimer1, InGPURepetitions, GPUCopyTimer, GPUNetworkTimer);
+		TEXT("Profiling for Run(%s/%s):\t1 time = %f+%f msec, avg(%d times) = %f+%f msec."),
+		*DeviceTypeString, *BackEndForCurrentPlatformString, CopyTimer1, NetworkTimer1, InRepetitions, CopyTimer, NetworkTimer);
+	// Reset device & back end
+	InOutNetwork->SetDeviceType(OriginalDeviceType);
+	InOutNetwork->SetBackEnd(OriginalBackEnd);
 	// Test successful
 	return true;
 }
diff --git a/Engine/Plugins/Experimental/NeuralNetworkInference/Source/NeuralNetworkInferenceQA/Private/ModelUnitTester.h b/Engine/Plugins/Experimental/NeuralNetworkInference/Source/NeuralNetworkInferenceQA/Private/ModelUnitTester.h
--- a/Engine/Plugins/Experimental/NeuralNetworkInference/Source/NeuralNetworkInferenceQA/Private/ModelUnitTester.h
+++ b/Engine/Plugins/Experimental/NeuralNetworkInference/Source/NeuralNetworkInferenceQA/Private/ModelUnitTester.h
@@ -16,7 +16,7 @@ private:
 	 */
 	static bool ModelLoadAccuracyAndSpeedTests(const FString& InProjectContentDir, const FString& InModelZooRelativeDirectory, const TArray<FString>& InModelNames,
 		const TArray<float>& InInputArrayValues, const TArray<TArray<double>>& InCPUGroundTruths, const TArray<TArray<double>>& InGPUGroundTruths,
-		const TArray<int32>& InCPURepetitions, const TArray<int32>& InGPURepetitions);
+		const TArray<int32>& InCPURepetitionsForUEAndORT, const TArray<int32>& InGPURepetitionsForUEAndORT, const TArray<int32>& InCPURepetitionsForUEOnly, const TArray<int32>& InGPURepetitionsForUEOnly);
 	/**
 	 * Other auxiliary functions for GlobalTest().
 	 */
@@ -27,7 +27,7 @@ private:
 	static UNeuralNetwork* NetworkONNXOrORTLoadTest(const FString& InModelFilePath);
 	static bool ModelAccuracyTest(UNeuralNetwork* InOutNetwork, const TArray<float>& InInputArrayValues, const TArray<double>& InCPUGroundTruths,
 		const TArray<double>& InGPUGroundTruths);
-	static bool ModelSpeedTest(const FString& InUAssetPath, const int32 InCPURepetitions, const int32 InGPURepetitions);
+	static bool ModelSpeedTest(const FString& InUAssetPath, const ENeuralDeviceType InDeviceType, const ENeuralBackEnd InBackEnd, const int32 InRepetitions);
 	static double GetAveragedL1Norm(const TArray<float>& InArray);
 	static double GetAveragedL1NormDiff(const TArray<float>& InArray1, const TArray<float>& InArray2);
 };
diff --git a/Engine/Plugins/Experimental/NeuralNetworkInference/Source/NeuralNetworkInferenceQA/Private/UnitTester.cpp b/Engine/Plugins/Experimental/NeuralNetworkInference/Source/NeuralNetworkInferenceQA/Private/UnitTester.cpp
--- a/Engine/Plugins/Experimental/NeuralNetworkInference/Source/NeuralNetworkInferenceQA/Private/UnitTester.cpp
+++ b/Engine/Plugins/Experimental/NeuralNetworkInference/Source/NeuralNetworkInferenceQA/Private/UnitTester.cpp
@@ -56,6 +56,7 @@ bool FUnitTester::GlobalTest(const FString& InProjectContentDir, const FString&
 	else
 	{
 		UE_LOG(LogNeuralNetworkInferenceQA, Warning, TEXT("----- UnitTesting() finished with warnings/errors! --------------------------------------------------------------------------------------------------"));
+		ensureMsgf(false, TEXT("UnitTesting() failed. See above log for more details."));
 	}
 
 	return bDidGlobalTestPassed;
