NeuralNetworkInference plugin:
- Bug fixed: FModelUnitTester::GlobalTest() does not break demo next time it is run.

#rnx

diff --git a/Engine/Plugins/Experimental/NeuralNetworkInference/Source/NeuralNetworkInferenceQA/Private/ModelUnitTester.cpp b/Engine/Plugins/Experimental/NeuralNetworkInference/Source/NeuralNetworkInferenceQA/Private/ModelUnitTester.cpp
--- a/Engine/Plugins/Experimental/NeuralNetworkInference/Source/NeuralNetworkInferenceQA/Private/ModelUnitTester.cpp
+++ b/Engine/Plugins/Experimental/NeuralNetworkInference/Source/NeuralNetworkInferenceQA/Private/ModelUnitTester.cpp
@@ -194,6 +194,9 @@ bool FModelUnitTester::ModelAccuracyTest(UNeuralNetwork* InOutNetwork, const TAr
 			return false;
 		}
 	}
+	// Save original network state
+	const ENeuralDeviceType OriginalDeviceType = InOutNetwork->GetDeviceType();
+	const ENeuralBackEnd OriginalBackEnd = InOutNetwork->GetBackEnd();
 	// Run each input with CPU/GPU and compare with each other and with the ground truth
 	bool bDidGlobalTestPassed = true;
 	for (int32 Index = 0; Index < InputArrays.Num(); ++Index)
@@ -244,6 +247,9 @@ bool FModelUnitTester::ModelAccuracyTest(UNeuralNetwork* InOutNetwork, const TAr
 		}
 		bDidGlobalTestPassed &= !bDidSomeTestFailed;
 	}
+	// Reset to original network state
+	InOutNetwork->SetDeviceType(OriginalDeviceType);
+	InOutNetwork->SetBackEnd(OriginalBackEnd);
 	// Test successful
 	return bDidGlobalTestPassed;
 }
@@ -320,7 +326,7 @@ bool FModelUnitTester::ModelSpeedTest(const FString& InUAssetPath, const ENeural
 		UE_LOG(LogNeuralNetworkInferenceQA, Warning, TEXT("FModelUnitTester::ModelSpeedTest(): InOutNetwork was a nullptr. Path: \"%s\"."), *InUAssetPath);
 		return false;
 	}
-	// Get original values
+	// Save original network state
 	const ENeuralDeviceType OriginalDeviceType = InOutNetwork->GetDeviceType();
 	const ENeuralBackEnd OriginalBackEnd = InOutNetwork->GetBackEnd();
 	// Set desired back end
@@ -342,7 +348,7 @@ bool FModelUnitTester::ModelSpeedTest(const FString& InUAssetPath, const ENeural
 	UE_LOG(LogNeuralNetworkInferenceQA, Display,
 		TEXT("Profiling for Run(%s/%s):\t1 time = %f+%f msec, avg(%d times) = %f+%f msec."),
 		*DeviceTypeString, *BackEndForCurrentPlatformString, CopyTimer1, NetworkTimer1, InRepetitions, CopyTimer, NetworkTimer);
-	// Reset device & back end
+	// Reset to original network state
 	InOutNetwork->SetDeviceType(OriginalDeviceType);
 	InOutNetwork->SetBackEnd(OriginalBackEnd);
 	// Test successful
