NeuralNetworkInference plugin:
- Code cleanup of ModelUnitTester.cpp

#rnx

diff --git a/Engine/Plugins/Experimental/NeuralNetworkInference/Source/NeuralNetworkInferenceQA/Private/ModelUnitTester.cpp b/Engine/Plugins/Experimental/NeuralNetworkInference/Source/NeuralNetworkInferenceQA/Private/ModelUnitTester.cpp
--- a/Engine/Plugins/Experimental/NeuralNetworkInference/Source/NeuralNetworkInferenceQA/Private/ModelUnitTester.cpp
+++ b/Engine/Plugins/Experimental/NeuralNetworkInference/Source/NeuralNetworkInferenceQA/Private/ModelUnitTester.cpp
@@ -196,6 +196,8 @@ bool FModelUnitTester::ModelAccuracyTest(UNeuralNetwork* InOutNetwork, const TAr
 	}
 	// Save original network state
 	const ENeuralDeviceType OriginalDeviceType = InOutNetwork->GetDeviceType();
+	const ENeuralDeviceType OriginalInputDeviceType = InOutNetwork->GetInputDeviceType();
+	const ENeuralDeviceType OriginalOutputDeviceType = InOutNetwork->GetOutputDeviceType();
 	const ENeuralBackEnd OriginalBackEnd = InOutNetwork->GetBackEnd();
 	// Run each input with CPU/GPU and compare with each other and with the ground truth
 	bool bDidGlobalTestPassed = true;
@@ -208,22 +210,29 @@ bool FModelUnitTester::ModelAccuracyTest(UNeuralNetwork* InOutNetwork, const TAr
 		// CPU
 		InOutNetwork->SetDeviceType(ENeuralDeviceType::CPU);
 		InOutNetwork->Run();
-		TArray<float> OutputArrayCPU = InOutNetwork->GetOutputTensor().GetArrayCopy<float>();
-		// GPU
+		TArray<float> CPUOutput = InOutNetwork->GetOutputTensor().GetArrayCopy<float>();
+		// Input CPU + GPU + Output CPU
 		InOutNetwork->SetDeviceType(ENeuralDeviceType::GPU);
 		InOutNetwork->Run();
-		TArray<float> OutputArrayGPU = InOutNetwork->GetOutputTensor().GetArrayCopy<float>();
+		TArray<float> CPUGPUCPUOutput = InOutNetwork->GetOutputTensor().GetArrayCopy<float>();
+		//// Input CPU + GPU + Output GPU
+		//InOutNetwork->SetDeviceType(ENeuralDeviceType::GPU);
+		//InOutNetwork->SetOutputDeviceType(ENeuralDeviceType::GPU);
+		//InOutNetwork->Run();
+		//TArray<float> CPUOutputGPUGPU = InOutNetwork->GetOutputTensor().GetArrayCopy<float>();
 		// Final verbose
-		const double CPUAvgL1Norm = GetAveragedL1Norm(OutputArrayCPU);
-		const double GPUAvgL1Norm = GetAveragedL1Norm(OutputArrayGPU);
-		const double RelativeCoefficient = 1. / FMath::Max(1., FMath::Min(CPUAvgL1Norm, GPUAvgL1Norm)); // Max(1, X) to avoid 0s
-		const double CPUGPUAvgL1NormDiff = GetAveragedL1NormDiff(OutputArrayCPU, OutputArrayGPU) * RelativeCoefficient * 1e3;
-		const double FastCPUGPUAvgL1NormDiff = FMath::Abs((CPUAvgL1Norm - GPUAvgL1Norm)) * RelativeCoefficient * 1e6;
+		const double CPUAvgL1Norm = GetAveragedL1Norm(CPUOutput);
+		const double CPUGPUCPUAvgL1Norm = GetAveragedL1Norm(CPUGPUCPUOutput);
+		//const double CPUGPUGPUAvgL1Norm = GetAveragedL1Norm(CPUOutputGPUGPU);
+		const double RelativeCoefficient = 1. / FMath::Max(1., FMath::Min(CPUAvgL1Norm, CPUGPUCPUAvgL1Norm)); // Max(1, X) to avoid 0s
+		const double CPUGPUAvgL1NormDiff = GetAveragedL1NormDiff(CPUOutput, CPUGPUCPUOutput) * RelativeCoefficient * 1e3;
+		const double FastCPUGPUAvgL1NormDiff = FMath::Abs((CPUAvgL1Norm - CPUGPUCPUAvgL1Norm)) * RelativeCoefficient * 1e6;
 		const double FastCPUAvgL1NormDiff = FMath::Abs(CPUAvgL1Norm - CPUGroundTruth) / FMath::Max(1., FMath::Min(CPUAvgL1Norm, CPUGroundTruth)) * 1e7;
-		const double FastGPUAvgL1NormDiff = FMath::Abs(GPUAvgL1Norm - GPUGroundTruth) / FMath::Max(1., FMath::Min(GPUAvgL1Norm, GPUGroundTruth)) * 1e7;
-		UE_LOG(LogNeuralNetworkInferenceQA, Display, TEXT("InputNorm = %f, OutputNormCPU = %f, OutputNormGPU = %f, OutputNormGT = %f, CPUAvgL1Norm = %f, GPUAvgL1Norm = %f,"
-			"\tCPUGPUAvgL1NormDiff = %fe-3, FastCPUGPUAvgL1NormDiff = %fe-6, FastCPUAvgL1NormDiff = %fe-7, FastGPUAvgL1NormDiff = %fe-7 (1e-7 is roughly the precision for float)."),
-			GetAveragedL1Norm(InputArray), CPUAvgL1Norm, GPUAvgL1Norm, CPUGroundTruth, CPUAvgL1Norm, GPUAvgL1Norm, CPUGPUAvgL1NormDiff, FastCPUGPUAvgL1NormDiff, FastCPUAvgL1NormDiff, FastGPUAvgL1NormDiff);
+		const double FastGPUAvgL1NormDiff = FMath::Abs(CPUGPUCPUAvgL1Norm - GPUGroundTruth) / FMath::Max(1., FMath::Min(CPUGPUCPUAvgL1Norm, GPUGroundTruth)) * 1e7;
+		UE_LOG(LogNeuralNetworkInferenceQA, Display, TEXT("InputNorm = %f, OutputNormCPU = %f, OutputNormGPU = %f, OutputNormGT = %f, CPUAvgL1Norm = %f, CPUGPUCPUAvgL1Norm = %f,"),
+			GetAveragedL1Norm(InputArray), CPUAvgL1Norm, CPUGPUCPUAvgL1Norm, CPUGroundTruth, CPUAvgL1Norm, CPUGPUCPUAvgL1Norm);
+		UE_LOG(LogNeuralNetworkInferenceQA, Display, TEXT("\tCPUGPUAvgL1NormDiff = %fe-3, FastCPUGPUAvgL1NormDiff = %fe-6, FastCPUAvgL1NormDiff = %fe-7, FastGPUAvgL1NormDiff = %fe-7 (1e-7 is roughly the precision for float)."),
+			CPUGPUAvgL1NormDiff, FastCPUGPUAvgL1NormDiff, FastCPUAvgL1NormDiff, FastGPUAvgL1NormDiff);
 		const TArray<int64>& InputSizes = InOutNetwork->GetInputTensor().GetSizes();
 		const TArray<int64>& OutputSizes = InOutNetwork->GetOutputTensor().GetSizes();
 		const int64 MaxNumberElementsToDisplay = 100;
@@ -236,12 +245,13 @@ bool FModelUnitTester::ModelAccuracyTest(UNeuralNetwork* InOutNetwork, const TAr
 		{
 			UE_LOG(LogNeuralNetworkInferenceQA, Display, TEXT("FastCPUGPUAvgL1NormDiff (%fe-6) < 5e-6 might have failed."), FastCPUGPUAvgL1NormDiff);
 			UE_LOG(LogNeuralNetworkInferenceQA, Display, TEXT("CPUGPUAvgL1NormDiff (%fe-7) < 1e-7 might have failed."), CPUGPUAvgL1NormDiff);
-			UE_LOG(LogNeuralNetworkInferenceQA, Display, TEXT("FastCPUAvgL1NormDiff (%fe-7) < 30e-7 might have failed (~30 times the float precision).\nOutputArrayCPU = %s."), FastCPUAvgL1NormDiff, *FNeuralTensor(OutputArrayCPU, OutputSizes).ToString(MaxNumberElementsToDisplay));
-			UE_LOG(LogNeuralNetworkInferenceQA, Display, TEXT("FastGPUAvgL1NormDiff (%fe-7) < 30e-7 might have failed (~30 times the float precision).\nOutputArrayGPU = %s."), FastGPUAvgL1NormDiff, *FNeuralTensor(OutputArrayGPU, OutputSizes).ToString(MaxNumberElementsToDisplay));
+			UE_LOG(LogNeuralNetworkInferenceQA, Display, TEXT("FastCPUAvgL1NormDiff (%fe-7) < 30e-7 might have failed (~30 times the float precision).\nCPUOutput = %s."), FastCPUAvgL1NormDiff, *FNeuralTensor(CPUOutput, OutputSizes).ToString(MaxNumberElementsToDisplay));
+			UE_LOG(LogNeuralNetworkInferenceQA, Display, TEXT("FastGPUAvgL1NormDiff (%fe-7) < 30e-7 might have failed (~30 times the float precision).\nCPUGPUCPUOutput = %s."), FastGPUAvgL1NormDiff, *FNeuralTensor(CPUGPUCPUOutput, OutputSizes).ToString(MaxNumberElementsToDisplay));
+			//UE_LOG(LogNeuralNetworkInferenceQA, Display, TEXT("FastGPUAvgL1NormDiff (%fe-7) < 30e-7 might have failed (~30 times the float precision).\nCPUOutputGPUGPU = %s."), FastGPUAvgL1NormDiff, *FNeuralTensor(CPUOutputGPUGPU, OutputSizes).ToString(MaxNumberElementsToDisplay));
 			UE_LOG(LogNeuralNetworkInferenceQA, Display, TEXT("Input = %s"),
 				*FNeuralTensor(InOutNetwork->GetInputTensor().GetArrayCopy<float>(), InputSizes).ToString(MaxNumberElementsToDisplay));
-			UE_LOG(LogNeuralNetworkInferenceQA, Display, TEXT("OutputArrayCPU = %s"), *FNeuralTensor(OutputArrayCPU, OutputSizes).ToString(MaxNumberElementsToDisplay));
-			UE_LOG(LogNeuralNetworkInferenceQA, Display, TEXT("OutputArrayGPU = %s"), *FNeuralTensor(OutputArrayGPU, OutputSizes).ToString(MaxNumberElementsToDisplay));
+			UE_LOG(LogNeuralNetworkInferenceQA, Display, TEXT("CPUOutput = %s"), *FNeuralTensor(CPUOutput, OutputSizes).ToString(MaxNumberElementsToDisplay));
+			UE_LOG(LogNeuralNetworkInferenceQA, Display, TEXT("CPUGPUCPUOutput = %s"), *FNeuralTensor(CPUGPUCPUOutput, OutputSizes).ToString(MaxNumberElementsToDisplay));
 			UE_LOG(LogNeuralNetworkInferenceQA, Warning, TEXT("At least 1 of the 4 CPU/GPU tests failed."));
 			return false;
 		}
@@ -249,6 +259,8 @@ bool FModelUnitTester::ModelAccuracyTest(UNeuralNetwork* InOutNetwork, const TAr
 	}
 	// Reset to original network state
 	InOutNetwork->SetDeviceType(OriginalDeviceType);
+	InOutNetwork->SetInputDeviceType(OriginalInputDeviceType);
+	InOutNetwork->SetOutputDeviceType(OriginalOutputDeviceType);
 	InOutNetwork->SetBackEnd(OriginalBackEnd);
 	// Test successful
 	return bDidGlobalTestPassed;
@@ -260,7 +272,7 @@ bool FModelUnitTester::ModelAccuracyTest(UNeuralNetwork* InOutNetwork, const TAr
 	for (int32 TimerIndex = 0; TimerIndex < InRepetitions; ++TimerIndex) \
 	{ \
 		InOutNetwork->SetInputFromArrayCopy(InputArray); \
-		OutputArrayGPU = InOutNetwork->GetOutputTensor().GetArrayCopy<float>(); \
+		CPUGPUCPUOutput = InOutNetwork->GetOutputTensor().GetArrayCopy<float>(); \
 	} \
 	const float OutCopyTimeInMilliSeconds = Timer.Toc() / InRepetitions; \
 	/* Forward() speed */ \
@@ -278,7 +290,7 @@ bool FModelUnitTester::ModelAccuracyTest(UNeuralNetwork* InOutNetwork, const TAr
 		{ \
 			InOutNetwork->SetInputFromArrayCopy(InputArray); \
 			InOutNetwork->Run(); \
-			OutputArrayGPU = InOutNetwork->GetOutputTensor().GetArrayCopy<float>(); \
+			CPUGPUCPUOutput = InOutNetwork->GetOutputTensor().GetArrayCopy<float>(); \
 		} \
 	} \
 	const float OutNetworkTimeInMilliSeconds = Timer.Toc() / InRepetitions - OutCopyTimeInMilliSeconds
@@ -335,7 +347,7 @@ bool FModelUnitTester::ModelSpeedTest(const FString& InUAssetPath, const ENeural
 	const int64 NetworkSize = InOutNetwork->GetInputTensor().Num();
 	TArray<float> InputArray;
 	InputArray.Init(1.f, NetworkSize);
-	TArray<float> OutputArrayGPU;
+	TArray<float> CPUGPUCPUOutput;
 	// Speed profiling
 	FNeuralNetworkInferenceQATimer Timer;
 	// Run profiling 1 time
