NeuralNetworkInference (NNI) plugin:
- Bug fixed: SetInput/OutputDevice should force a new Load() for ORT, like SetDevice() does.
- Added UNN::SetXDeviceType() for code reuse.

#rnx

diff --git a/Engine/Plugins/Experimental/NeuralNetworkInference/Source/NeuralNetworkInference/Private/NeuralNetwork.cpp b/Engine/Plugins/Experimental/NeuralNetworkInference/Source/NeuralNetworkInference/Private/NeuralNetwork.cpp
--- a/Engine/Plugins/Experimental/NeuralNetworkInference/Source/NeuralNetworkInference/Private/NeuralNetwork.cpp
+++ b/Engine/Plugins/Experimental/NeuralNetworkInference/Source/NeuralNetworkInference/Private/NeuralNetwork.cpp
@@ -127,14 +127,7 @@ ENeuralDeviceType UNeuralNetwork::GetDeviceType() const
 
 void UNeuralNetwork::SetDeviceType(const ENeuralDeviceType InDeviceType)
 {
-	if (DeviceType != InDeviceType)
-	{
-		DeviceType = InDeviceType;
-		if (BackEndForCurrentPlatform == ENeuralBackEnd::UEAndORT)
-		{
-			Load();
-		}
-	}
+	SetXDeviceType(DeviceType, InDeviceType);
 }
 
 ENeuralDeviceType UNeuralNetwork::GetInputDeviceType() const
@@ -144,7 +137,7 @@ ENeuralDeviceType UNeuralNetwork::GetInputDeviceType() const
 
 void UNeuralNetwork::SetInputDeviceType(const ENeuralDeviceType InInputDeviceType)
 {
-	InputDeviceType = InInputDeviceType;
+	SetXDeviceType(InputDeviceType, InInputDeviceType);
 }
 
 ENeuralDeviceType UNeuralNetwork::GetOutputDeviceType() const
@@ -154,7 +147,7 @@ ENeuralDeviceType UNeuralNetwork::GetOutputDeviceType() const
 
 void UNeuralNetwork::SetOutputDeviceType(const ENeuralDeviceType InOutputDeviceType)
 {
-	OutputDeviceType = InOutputDeviceType;
+	SetXDeviceType(OutputDeviceType, InOutputDeviceType);
 }
 
 ENeuralNetworkSynchronousMode UNeuralNetwork::GetSynchronousMode() const
@@ -388,6 +381,18 @@ void UNeuralNetwork::Run()
 /* UNeuralNetwork private functions
  *****************************************************************************/
 
+void UNeuralNetwork::SetXDeviceType(ENeuralDeviceType& OutDeviceType, const ENeuralDeviceType InDeviceType)
+{
+	if (OutDeviceType != InDeviceType)
+	{
+		OutDeviceType = InDeviceType;
+		if (BackEndForCurrentPlatform == ENeuralBackEnd::UEAndORT)
+		{
+			Load();
+		}
+	}
+}
+
 bool UNeuralNetwork::Load()
 {
 	DECLARE_SCOPE_CYCLE_COUNTER(TEXT("UNeuralNetwork_Load"), STAT_UNeuralNetwork_Load, STATGROUP_MachineLearning);
diff --git a/Engine/Plugins/Experimental/NeuralNetworkInference/Source/NeuralNetworkInference/Public/NeuralNetwork.h b/Engine/Plugins/Experimental/NeuralNetworkInference/Source/NeuralNetworkInference/Public/NeuralNetwork.h
--- a/Engine/Plugins/Experimental/NeuralNetworkInference/Source/NeuralNetworkInference/Public/NeuralNetwork.h
+++ b/Engine/Plugins/Experimental/NeuralNetworkInference/Source/NeuralNetworkInference/Public/NeuralNetwork.h
@@ -229,6 +229,11 @@ private:
 	struct FImplBackEndUEOnly;
 	TSharedPtr<FImplBackEndUEOnly> ImplBackEndUEOnly;
 
+	/**
+	 * Auxiliary function for Set(Input/Output)DeviceType.
+	 */
+	void SetXDeviceType(ENeuralDeviceType& OutDeviceType, const ENeuralDeviceType InDeviceType);
+
 	/**
 	 * It loads the desired network graph definition and weights internally saved on this UNeuralNetwork instance.
 	 * @return Whether the network was successfully loaded.
