NeuralNetworkInference plugin:
- Removed redundant ensure(false).

#rnx

diff --git a/Engine/Plugins/Experimental/NeuralNetworkInference/Source/NeuralNetworkInferenceQA/Private/UnitTester.cpp b/Engine/Plugins/Experimental/NeuralNetworkInference/Source/NeuralNetworkInferenceQA/Private/UnitTester.cpp
--- a/Engine/Plugins/Experimental/NeuralNetworkInference/Source/NeuralNetworkInferenceQA/Private/UnitTester.cpp
+++ b/Engine/Plugins/Experimental/NeuralNetworkInference/Source/NeuralNetworkInferenceQA/Private/UnitTester.cpp
@@ -56,7 +56,6 @@ bool FUnitTester::GlobalTest(const FString& InProjectContentDir, const FString&
 	else
 	{
 		UE_LOG(LogNeuralNetworkInferenceQA, Warning, TEXT("----- UnitTesting() finished with warnings/errors! --------------------------------------------------------------------------------------------------"));
-		ensureMsgf(false, TEXT("UnitTesting() failed. See above log for more details."));
 	}
 
 	return bDidGlobalTestPassed;
