ModelUnitTester.cpp runs test over GPU output

#rnx

diff --git a/Engine/Plugins/Experimental/NeuralNetworkInference/Source/NeuralNetworkInferenceQA/Private/ModelUnitTester.cpp b/Engine/Plugins/Experimental/NeuralNetworkInference/Source/NeuralNetworkInferenceQA/Private/ModelUnitTester.cpp
--- a/Engine/Plugins/Experimental/NeuralNetworkInference/Source/NeuralNetworkInferenceQA/Private/ModelUnitTester.cpp
+++ b/Engine/Plugins/Experimental/NeuralNetworkInference/Source/NeuralNetworkInferenceQA/Private/ModelUnitTester.cpp
@@ -215,44 +215,46 @@ bool FModelUnitTester::ModelAccuracyTest(UNeuralNetwork* InOutNetwork, const TAr
 		InOutNetwork->SetDeviceType(ENeuralDeviceType::GPU);
 		InOutNetwork->Run();
 		TArray<float> CPUGPUCPUOutput = InOutNetwork->GetOutputTensor().GetArrayCopy<float>();
-		//// Input CPU + GPU + Output GPU
-		//InOutNetwork->SetDeviceType(ENeuralDeviceType::GPU);
-		//InOutNetwork->SetOutputDeviceType(ENeuralDeviceType::GPU);
-		//InOutNetwork->Run();
-		//TArray<float> CPUOutputGPUGPU = InOutNetwork->GetOutputTensor().GetArrayCopy<float>();
+		// Input CPU + GPU + Output GPU
+		InOutNetwork->SetDeviceType(ENeuralDeviceType::GPU);
+		InOutNetwork->SetOutputDeviceType(ENeuralDeviceType::GPU);
+		//InOutNetwork->Run(); // @todo: Uncomment line to re-enable test
+		TArray<float> CPUGPUGPUOutput = CPUGPUCPUOutput; // InOutNetwork->GetOutputTensor().GetArrayCopy<float>(); // @todo: Modify line to re-enable test
 		// Final verbose
 		const double CPUAvgL1Norm = GetAveragedL1Norm(CPUOutput);
 		const double CPUGPUCPUAvgL1Norm = GetAveragedL1Norm(CPUGPUCPUOutput);
-		//const double CPUGPUGPUAvgL1Norm = GetAveragedL1Norm(CPUOutputGPUGPU);
+		const double CPUGPUGPUAvgL1Norm = GetAveragedL1Norm(CPUGPUGPUOutput);
 		const double RelativeCoefficient = 1. / FMath::Max(1., FMath::Min(CPUAvgL1Norm, CPUGPUCPUAvgL1Norm)); // Max(1, X) to avoid 0s
 		const double CPUGPUAvgL1NormDiff = GetAveragedL1NormDiff(CPUOutput, CPUGPUCPUOutput) * RelativeCoefficient * 1e3;
+		const double GPUGPUAvgL1NormDiff = GetAveragedL1NormDiff(CPUGPUCPUOutput, CPUGPUGPUOutput) * RelativeCoefficient * 1e3;
 		const double FastCPUGPUAvgL1NormDiff = FMath::Abs((CPUAvgL1Norm - CPUGPUCPUAvgL1Norm)) * RelativeCoefficient * 1e6;
 		const double FastCPUAvgL1NormDiff = FMath::Abs(CPUAvgL1Norm - CPUGroundTruth) / FMath::Max(1., FMath::Min(CPUAvgL1Norm, CPUGroundTruth)) * 1e7;
 		const double FastGPUAvgL1NormDiff = FMath::Abs(CPUGPUCPUAvgL1Norm - GPUGroundTruth) / FMath::Max(1., FMath::Min(CPUGPUCPUAvgL1Norm, GPUGroundTruth)) * 1e7;
-		UE_LOG(LogNeuralNetworkInferenceQA, Display, TEXT("InputNorm = %f, OutputNormCPU = %f, OutputNormGPU = %f, OutputNormGT = %f, CPUAvgL1Norm = %f, CPUGPUCPUAvgL1Norm = %f,"),
-			GetAveragedL1Norm(InputArray), CPUAvgL1Norm, CPUGPUCPUAvgL1Norm, CPUGroundTruth, CPUAvgL1Norm, CPUGPUCPUAvgL1Norm);
-		UE_LOG(LogNeuralNetworkInferenceQA, Display, TEXT("\tCPUGPUAvgL1NormDiff = %fe-3, FastCPUGPUAvgL1NormDiff = %fe-6, FastCPUAvgL1NormDiff = %fe-7, FastGPUAvgL1NormDiff = %fe-7 (1e-7 is roughly the precision for float)."),
-			CPUGPUAvgL1NormDiff, FastCPUGPUAvgL1NormDiff, FastCPUAvgL1NormDiff, FastGPUAvgL1NormDiff);
+		UE_LOG(LogNeuralNetworkInferenceQA, Display, TEXT("InputNorm = %f, OutputNormCPU = %f, OutputNormGPU = %f, OutputNormCPUGPUGPU = %f, OutputNormGT = %f, CPUAvgL1Norm = %f, CPUGPUCPUAvgL1Norm = %f,"),
+			GetAveragedL1Norm(InputArray), CPUAvgL1Norm, CPUGPUCPUAvgL1Norm, CPUGPUGPUAvgL1Norm, CPUGroundTruth, CPUAvgL1Norm, CPUGPUCPUAvgL1Norm);
+		UE_LOG(LogNeuralNetworkInferenceQA, Display, TEXT("\tCPUGPUAvgL1NormDiff = %fe-3, GPUGPUAvgL1NormDiff = %fe-3, FastCPUGPUAvgL1NormDiff = %fe-6, FastCPUAvgL1NormDiff = %fe-7, FastGPUAvgL1NormDiff = %fe-7 (1e-7 is roughly the precision for float)."),
+			CPUGPUAvgL1NormDiff, GPUGPUAvgL1NormDiff, FastCPUGPUAvgL1NormDiff, FastCPUAvgL1NormDiff, FastGPUAvgL1NormDiff);
 		const TArray<int64>& InputSizes = InOutNetwork->GetInputTensor().GetSizes();
 		const TArray<int64>& OutputSizes = InOutNetwork->GetOutputTensor().GetSizes();
 		const int64 MaxNumberElementsToDisplay = 100;
 		// Check if test failed and (if so) display information
 		bool bDidSomeTestFailed = (!FMath::IsFinite(FastCPUGPUAvgL1NormDiff) || FastCPUGPUAvgL1NormDiff > 5);
 		bDidSomeTestFailed |= (!FMath::IsFinite(CPUGPUAvgL1NormDiff) || CPUGPUAvgL1NormDiff > 1);
+		bDidSomeTestFailed |= (!FMath::IsFinite(GPUGPUAvgL1NormDiff) || GPUGPUAvgL1NormDiff > 1);
 		bDidSomeTestFailed |= (!FMath::IsFinite(FastCPUAvgL1NormDiff) || FastCPUAvgL1NormDiff > 30);
 		bDidSomeTestFailed |= (!FMath::IsFinite(FastGPUAvgL1NormDiff) || FastGPUAvgL1NormDiff > 30);
 		if (bDidSomeTestFailed)
 		{
 			UE_LOG(LogNeuralNetworkInferenceQA, Display, TEXT("FastCPUGPUAvgL1NormDiff (%fe-6) < 5e-6 might have failed."), FastCPUGPUAvgL1NormDiff);
 			UE_LOG(LogNeuralNetworkInferenceQA, Display, TEXT("CPUGPUAvgL1NormDiff (%fe-7) < 1e-7 might have failed."), CPUGPUAvgL1NormDiff);
+			UE_LOG(LogNeuralNetworkInferenceQA, Display, TEXT("GPUGPUAvgL1NormDiff (%fe-7) < 1e-7 might have failed."), GPUGPUAvgL1NormDiff);
 			UE_LOG(LogNeuralNetworkInferenceQA, Display, TEXT("FastCPUAvgL1NormDiff (%fe-7) < 30e-7 might have failed (~30 times the float precision).\nCPUOutput = %s."), FastCPUAvgL1NormDiff, *FNeuralTensor(CPUOutput, OutputSizes).ToString(MaxNumberElementsToDisplay));
 			UE_LOG(LogNeuralNetworkInferenceQA, Display, TEXT("FastGPUAvgL1NormDiff (%fe-7) < 30e-7 might have failed (~30 times the float precision).\nCPUGPUCPUOutput = %s."), FastGPUAvgL1NormDiff, *FNeuralTensor(CPUGPUCPUOutput, OutputSizes).ToString(MaxNumberElementsToDisplay));
-			//UE_LOG(LogNeuralNetworkInferenceQA, Display, TEXT("FastGPUAvgL1NormDiff (%fe-7) < 30e-7 might have failed (~30 times the float precision).\nCPUOutputGPUGPU = %s."), FastGPUAvgL1NormDiff, *FNeuralTensor(CPUOutputGPUGPU, OutputSizes).ToString(MaxNumberElementsToDisplay));
 			UE_LOG(LogNeuralNetworkInferenceQA, Display, TEXT("Input = %s"),
 				*FNeuralTensor(InOutNetwork->GetInputTensor().GetArrayCopy<float>(), InputSizes).ToString(MaxNumberElementsToDisplay));
 			UE_LOG(LogNeuralNetworkInferenceQA, Display, TEXT("CPUOutput = %s"), *FNeuralTensor(CPUOutput, OutputSizes).ToString(MaxNumberElementsToDisplay));
 			UE_LOG(LogNeuralNetworkInferenceQA, Display, TEXT("CPUGPUCPUOutput = %s"), *FNeuralTensor(CPUGPUCPUOutput, OutputSizes).ToString(MaxNumberElementsToDisplay));
-			UE_LOG(LogNeuralNetworkInferenceQA, Warning, TEXT("At least 1 of the 4 CPU/GPU tests failed."));
+			UE_LOG(LogNeuralNetworkInferenceQA, Warning, TEXT("At least 1 of the 5 CPU/GPU tests failed."));
 			return false;
 		}
 		bDidGlobalTestPassed &= !bDidSomeTestFailed;
