NeuralNetworkInference (NNI) plugin:
- ModelUnitTester.cpp runs accuracy tests over UEOnly back end as well.

#rnx

diff --git a/Engine/Plugins/Experimental/NeuralNetworkInference/Source/NeuralNetworkInferenceQA/Private/ModelUnitTester.cpp b/Engine/Plugins/Experimental/NeuralNetworkInference/Source/NeuralNetworkInferenceQA/Private/ModelUnitTester.cpp
--- a/Engine/Plugins/Experimental/NeuralNetworkInference/Source/NeuralNetworkInferenceQA/Private/ModelUnitTester.cpp
+++ b/Engine/Plugins/Experimental/NeuralNetworkInference/Source/NeuralNetworkInferenceQA/Private/ModelUnitTester.cpp
@@ -22,11 +22,11 @@ bool FModelUnitTester::GlobalTest(const FString& InProjectContentDir, const FStr
 		{138.373184, 126.754100, 127.287398, 130.316194, 127.303495, 124.801134, 126.5462530}, {0.488662, 0.472437, 0.478862, 0.522685, 0.038322, 0.480848, 0.483821} });
 	// Speed profiling test - 0 repetitions means that test will not be run
 #ifdef WITH_UE_AND_ORT_SUPPORT
-	const TArray<int32> CPURepetitionsForUEAndORT({ 1000, 1000,  50, 1000 });
-	const TArray<int32> GPURepetitionsForUEAndORT({ 1000, 1000, 100, 1000 });
+	const TArray<int32> CPURepetitionsForUEAndORTBackEnd({ 1000, 1000,  50, 1000 });
+	const TArray<int32> GPURepetitionsForUEAndORTBackEnd({ 1000, 1000, 100, 1000 });
 #else //WITH_UE_AND_ORT_SUPPORT
-	const TArray<int32> CPURepetitionsForUEAndORT({ 0, 0, 0, 0 });
-	const TArray<int32> GPURepetitionsForUEAndORT({ 0, 0, 0, 0 });
+	const TArray<int32> CPURepetitionsForUEAndORTBackEnd({ 0, 0, 0, 0 });
+	const TArray<int32> GPURepetitionsForUEAndORTBackEnd({ 0, 0, 0, 0 });
 #endif //WITH_UE_AND_ORT_SUPPORT
 #if WITH_EDITOR
 	const TArray<int32> CPURepetitionsForUEOnlyBackEnd({ 10, 0, 0, 0 });
@@ -36,7 +36,7 @@ bool FModelUnitTester::GlobalTest(const FString& InProjectContentDir, const FStr
 	const TArray<int32> GPURepetitionsForUEOnlyBackEnd({ 0, 0, 0, 0 });
 #endif //WITH_EDITOR
 	// Run tests
-	return ModelLoadAccuracyAndSpeedTests(InProjectContentDir, InModelZooRelativeDirectory, ModelNames, InputArrayValues, CPUGroundTruths, GPUGroundTruths, CPURepetitionsForUEAndORT, GPURepetitionsForUEAndORT, CPURepetitionsForUEOnlyBackEnd, GPURepetitionsForUEOnlyBackEnd);
+	return ModelLoadAccuracyAndSpeedTests(InProjectContentDir, InModelZooRelativeDirectory, ModelNames, InputArrayValues, CPUGroundTruths, GPUGroundTruths, CPURepetitionsForUEAndORTBackEnd, GPURepetitionsForUEAndORTBackEnd, CPURepetitionsForUEOnlyBackEnd, GPURepetitionsForUEOnlyBackEnd);
 }
 
 
@@ -45,8 +45,8 @@ bool FModelUnitTester::GlobalTest(const FString& InProjectContentDir, const FStr
  *****************************************************************************/
 
 bool FModelUnitTester::ModelLoadAccuracyAndSpeedTests(const FString& InProjectContentDir, const FString& InModelZooRelativeDirectory, const TArray<FString>& InModelNames,
-	const TArray<float>& InInputArrayValues, const TArray<TArray<double>>& InCPUGroundTruths, const TArray<TArray<double>>& InGPUGroundTruths, const TArray<int32>& InCPURepetitionsForUEAndORT,
-	const TArray<int32>& InGPURepetitionsForUEAndORT, const TArray<int32>& InCPURepetitionsForUEOnlyBackEnd, const TArray<int32>& InGPURepetitionsForUEOnlyBackEnd)
+	const TArray<float>& InInputArrayValues, const TArray<TArray<double>>& InCPUGroundTruths, const TArray<TArray<double>>& InGPUGroundTruths, const TArray<int32>& InCPURepetitionsForUEAndORTBackEnd,
+	const TArray<int32>& InGPURepetitionsForUEAndORTBackEnd, const TArray<int32>& InCPURepetitionsForUEOnlyBackEnd, const TArray<int32>& InGPURepetitionsForUEOnlyBackEnd)
 {
 	bool bDidGlobalTestPassed = true;
 
@@ -78,7 +78,16 @@ bool FModelUnitTester::ModelLoadAccuracyAndSpeedTests(const FString& InProjectCo
 		{
 			UE_LOG(LogNeuralNetworkInferenceQA, Display, TEXT("OutputTensor[%d] = %s."), TensorIndex, *Network->GetOutputTensor(TensorIndex).GetName());
 		}
-		bDidGlobalTestPassed &= ModelAccuracyTest(Network, InInputArrayValues, CPUGroundTruths, GPUGroundTruths);
+		const bool bShouldRunUEAndORTBackEnd = (InCPURepetitionsForUEAndORTBackEnd[ModelIndex] * InCPURepetitionsForUEAndORTBackEnd[ModelIndex] > 0);
+		if (bShouldRunUEAndORTBackEnd)
+		{
+			bDidGlobalTestPassed &= ModelAccuracyTest(Network, ENeuralBackEnd::UEAndORT, InInputArrayValues, CPUGroundTruths, GPUGroundTruths);
+		}
+		const bool bShouldRunUEOnlyBackEnd = (InCPURepetitionsForUEOnlyBackEnd[ModelIndex] * InCPURepetitionsForUEOnlyBackEnd[ModelIndex] > 0);
+		if (bShouldRunUEOnlyBackEnd)
+		{
+			bDidGlobalTestPassed &= ModelAccuracyTest(Network, ENeuralBackEnd::UEOnly, InInputArrayValues, CPUGroundTruths, GPUGroundTruths);
+		}
 
 		UE_LOG(LogNeuralNetworkInferenceQA, Display, TEXT("---------------------------------------------------------------------------------------------------------------------------------"));
 
@@ -89,7 +98,14 @@ bool FModelUnitTester::ModelLoadAccuracyAndSpeedTests(const FString& InProjectCo
 			const FString ModelFilePath = (ONNXOrORTIndex % 2 == 0 ? GetONNXModelFilePath(ModelZooDirectory, ModelName) : GetORTModelFilePath(ModelZooDirectory, ModelName));
 			const FString ModelType = (ONNXOrORTIndex % 2 == 0 ? TEXT("ONNX") : TEXT("ORT"));
 			UE_LOG(LogNeuralNetworkInferenceQA, Display, TEXT("-------------------- %s - Network %s Load and Run - %s"), *ModelName, *ModelType, *ModelFilePath);
-			bDidGlobalTestPassed &= ModelAccuracyTest(NetworkONNXOrORTLoadTest(ModelFilePath), InInputArrayValues, CPUGroundTruths, GPUGroundTruths);
+			if (bShouldRunUEAndORTBackEnd)
+			{
+				bDidGlobalTestPassed &= ModelAccuracyTest(NetworkONNXOrORTLoadTest(ModelFilePath), ENeuralBackEnd::UEAndORT, InInputArrayValues, CPUGroundTruths, GPUGroundTruths);
+			}
+			if (bShouldRunUEOnlyBackEnd)
+			{
+				bDidGlobalTestPassed &= ModelAccuracyTest(NetworkONNXOrORTLoadTest(ModelFilePath), ENeuralBackEnd::UEOnly, InInputArrayValues, CPUGroundTruths, GPUGroundTruths);
+			}
 		}
 #else //WITH_EDITOR
 	UE_LOG(LogNeuralNetworkInferenceQA, Display, TEXT("-------------------- Skipped (only if WITH_EDITOR enabled)."));
@@ -105,9 +121,9 @@ bool FModelUnitTester::ModelLoadAccuracyAndSpeedTests(const FString& InProjectCo
 
 		// UEAndORT (if WITH_UE_AND_ORT_SUPPORT)
 #ifdef WITH_UE_AND_ORT_SUPPORT
-		bDidGlobalTestPassed &= ModelSpeedTest(UAssetModelFilePath, ENeuralDeviceType::CPU, ENeuralBackEnd::UEAndORT, InCPURepetitionsForUEAndORT[ModelIndex]);
+		bDidGlobalTestPassed &= ModelSpeedTest(UAssetModelFilePath, ENeuralDeviceType::CPU, ENeuralBackEnd::UEAndORT, InCPURepetitionsForUEAndORTBackEnd[ModelIndex]);
 #ifdef PLATFORM_WIN64
-		bDidGlobalTestPassed &= ModelSpeedTest(UAssetModelFilePath, ENeuralDeviceType::GPU, ENeuralBackEnd::UEAndORT, InGPURepetitionsForUEAndORT[ModelIndex]);
+		bDidGlobalTestPassed &= ModelSpeedTest(UAssetModelFilePath, ENeuralDeviceType::GPU, ENeuralBackEnd::UEAndORT, InGPURepetitionsForUEAndORTBackEnd[ModelIndex]);
 #else //PLATFORM_WIN64
 		bDidGlobalTestPassed &= ModelSpeedTest(UAssetModelFilePath, ENeuralDeviceType::GPU, ENeuralBackEnd::UEAndORT, 0);
 #endif //PLATFORM_WIN64
@@ -169,7 +185,7 @@ UNeuralNetwork* FModelUnitTester::NetworkONNXOrORTLoadTest(const FString& InMode
 	return Network;
 }
 
-bool FModelUnitTester::ModelAccuracyTest(UNeuralNetwork* InOutNetwork, const TArray<float>& InInputArrayValues, const TArray<double>& InCPUGroundTruths, const TArray<double>& InGPUGroundTruths)
+bool FModelUnitTester::ModelAccuracyTest(UNeuralNetwork* InOutNetwork, const ENeuralBackEnd InBackEnd, const TArray<float>& InInputArrayValues, const TArray<double>& InCPUGroundTruths, const TArray<double>& InGPUGroundTruths)
 {
 	// Sanity check
 	if (!InOutNetwork)
@@ -199,6 +215,9 @@ bool FModelUnitTester::ModelAccuracyTest(UNeuralNetwork* InOutNetwork, const TAr
 	const ENeuralDeviceType OriginalInputDeviceType = InOutNetwork->GetInputDeviceType();
 	const ENeuralDeviceType OriginalOutputDeviceType = InOutNetwork->GetOutputDeviceType();
 	const ENeuralBackEnd OriginalBackEnd = InOutNetwork->GetBackEnd();
+	// Set back end
+	InOutNetwork->SetBackEnd(InBackEnd);
+	const FString BackEndString = GetBackEndString(InBackEnd);
 	// Run each input with CPU/GPU and compare with each other and with the ground truth
 	bool bDidGlobalTestPassed = true;
 	for (int32 Index = 0; Index < InputArrays.Num(); ++Index)
@@ -230,8 +249,8 @@ bool FModelUnitTester::ModelAccuracyTest(UNeuralNetwork* InOutNetwork, const TAr
 		const double FastCPUGPUAvgL1NormDiff = FMath::Abs((CPUAvgL1Norm - CPUGPUCPUAvgL1Norm)) * RelativeCoefficient * 1e6;
 		const double FastCPUAvgL1NormDiff = FMath::Abs(CPUAvgL1Norm - CPUGroundTruth) / FMath::Max(1., FMath::Min(CPUAvgL1Norm, CPUGroundTruth)) * 1e7;
 		const double FastGPUAvgL1NormDiff = FMath::Abs(CPUGPUCPUAvgL1Norm - GPUGroundTruth) / FMath::Max(1., FMath::Min(CPUGPUCPUAvgL1Norm, GPUGroundTruth)) * 1e7;
-		UE_LOG(LogNeuralNetworkInferenceQA, Display, TEXT("InputNorm = %f, OutputNormCPU = %f, OutputNormGPU = %f, OutputNormCPUGPUGPU = %f, OutputNormGT = %f, CPUAvgL1Norm = %f, CPUGPUCPUAvgL1Norm = %f,"),
-			GetAveragedL1Norm(InputArray), CPUAvgL1Norm, CPUGPUCPUAvgL1Norm, CPUGPUGPUAvgL1Norm, CPUGroundTruth, CPUAvgL1Norm, CPUGPUCPUAvgL1Norm);
+		UE_LOG(LogNeuralNetworkInferenceQA, Display, TEXT("%s: InputNorm = %f, OutputNormCPU = %f, OutputNormGPU = %f, OutputNormCPUGPUGPU = %f, OutputNormGT = %f, CPUAvgL1Norm = %f, CPUGPUCPUAvgL1Norm = %f,"),
+			*BackEndString, GetAveragedL1Norm(InputArray), CPUAvgL1Norm, CPUGPUCPUAvgL1Norm, CPUGPUGPUAvgL1Norm, CPUGroundTruth, CPUAvgL1Norm, CPUGPUCPUAvgL1Norm);
 		UE_LOG(LogNeuralNetworkInferenceQA, Display, TEXT("\tCPUGPUAvgL1NormDiff = %fe-3, GPUGPUAvgL1NormDiff = %fe-3, FastCPUGPUAvgL1NormDiff = %fe-6, FastCPUAvgL1NormDiff = %fe-7, FastGPUAvgL1NormDiff = %fe-7 (1e-7 is roughly the precision for float)."),
 			CPUGPUAvgL1NormDiff, GPUGPUAvgL1NormDiff, FastCPUGPUAvgL1NormDiff, FastCPUAvgL1NormDiff, FastGPUAvgL1NormDiff);
 		const TArray<int64>& InputSizes = InOutNetwork->GetInputTensor().GetSizes();
@@ -300,36 +319,12 @@ bool FModelUnitTester::ModelAccuracyTest(UNeuralNetwork* InOutNetwork, const TAr
 bool FModelUnitTester::ModelSpeedTest(const FString& InUAssetPath, const ENeuralDeviceType InDeviceType, const ENeuralBackEnd InBackEnd, const int32 InRepetitions)
 {
 	// Get debug strings
-	FString DeviceTypeString;
-	if (InDeviceType == ENeuralDeviceType::CPU)
-	{
-		DeviceTypeString = TEXT("CPU");
-	}
-	else if (InDeviceType == ENeuralDeviceType::GPU)
-	{
-		DeviceTypeString = TEXT("GPU");
-	}
-	else
-	{
-		UE_LOG(LogNeuralNetworkInferenceQA, Warning, TEXT("FModelUnitTester::ModelSpeedTest(): Unknown DeviceTypeString = %d."), (int32)InDeviceType);
-	}
-	FString BackEndForCurrentPlatformString;
-	if (InBackEnd == ENeuralBackEnd::UEAndORT)
-	{
-		BackEndForCurrentPlatformString = TEXT("UEAndORT");
-	}
-	else if (InBackEnd == ENeuralBackEnd::UEOnly)
-	{
-		BackEndForCurrentPlatformString = TEXT("UEOnly");
-	}
-	else
-	{
-		UE_LOG(LogNeuralNetworkInferenceQA, Warning, TEXT("FModelUnitTester::ModelSpeedTest(): Unknown BackEndForCurrentPlatformString = %d."), (int32)InBackEnd);
-	}
+	const FString DeviceTypeString = GetDeviceTypeString(InDeviceType);
+	const FString BackEndString = GetBackEndString(InBackEnd);
 	// Skip test if reps = 0
 	if (InRepetitions < 1)
 	{
-		UE_LOG(LogNeuralNetworkInferenceQA, Display, TEXT("ModelSpeedTest skipped for uasset \"%s\" for InDeviceType = %s (%d) and InBackEnd = %s (%d)."), *InUAssetPath, *DeviceTypeString, (int32)InDeviceType, *BackEndForCurrentPlatformString, (int32)InBackEnd);
+		UE_LOG(LogNeuralNetworkInferenceQA, Display, TEXT("ModelSpeedTest skipped for uasset \"%s\" for InDeviceType = %s (%d) and InBackEnd = %s (%d)."), *InUAssetPath, *DeviceTypeString, (int32)InDeviceType, *BackEndString, (int32)InBackEnd);
 		return true;
 	}
 	// Load Network
@@ -360,8 +355,8 @@ bool FModelUnitTester::ModelSpeedTest(const FString& InUAssetPath, const ENeural
 	MODEL_UNIT_TESTER_SPEED_TEST(CopyTimer, NetworkTimer, InRepetitions, /*bIsGPU*/false);
 	// Display speed times
 	UE_LOG(LogNeuralNetworkInferenceQA, Display,
-		TEXT("Profiling for Run(%s/%s):\t1 time = %f+%f msec, avg(%d times) = %f+%f msec."),
-		*DeviceTypeString, *BackEndForCurrentPlatformString, CopyTimer1, NetworkTimer1, InRepetitions, CopyTimer, NetworkTimer);
+		TEXT("%s-%s:\t1 time = %f+%f msec, avg(%d times) = %f+%f msec."),
+		*BackEndString, *DeviceTypeString, CopyTimer1, NetworkTimer1, InRepetitions, CopyTimer, NetworkTimer);
 	// Reset to original network state
 	InOutNetwork->SetDeviceType(OriginalDeviceType);
 	InOutNetwork->SetBackEnd(OriginalBackEnd);
@@ -397,3 +392,37 @@ double FModelUnitTester::GetAveragedL1NormDiff(const TArray<float>& InArray1, co
 	AveragedL1NormDiff /= InArray1.Num();
 	return AveragedL1NormDiff;
 }
+
+FString FModelUnitTester::GetDeviceTypeString(const ENeuralDeviceType InDeviceType)
+{
+	if (InDeviceType == ENeuralDeviceType::CPU)
+	{
+		return TEXT("CPU");
+	}
+	else if (InDeviceType == ENeuralDeviceType::GPU)
+	{
+		return TEXT("GPU");
+	}
+	else
+	{
+		ensureMsgf(false, TEXT("Unknown DeviceType = %d."), (int32)InDeviceType);
+		return TEXT("");
+	}
+}
+
+FString FModelUnitTester::GetBackEndString(const ENeuralBackEnd InBackEnd)
+{
+	if (InBackEnd == ENeuralBackEnd::UEAndORT)
+	{
+		return TEXT("UEAndORT");
+	}
+	else if (InBackEnd == ENeuralBackEnd::UEOnly)
+	{
+		return TEXT("UEOnly");
+	}
+	else
+	{
+		ensureMsgf(false, TEXT("Unknown BackEndForCurrentPlatform = %d."), (int32)InBackEnd);
+		return TEXT("");
+	}
+}
diff --git a/Engine/Plugins/Experimental/NeuralNetworkInference/Source/NeuralNetworkInferenceQA/Private/ModelUnitTester.h b/Engine/Plugins/Experimental/NeuralNetworkInference/Source/NeuralNetworkInferenceQA/Private/ModelUnitTester.h
--- a/Engine/Plugins/Experimental/NeuralNetworkInference/Source/NeuralNetworkInferenceQA/Private/ModelUnitTester.h
+++ b/Engine/Plugins/Experimental/NeuralNetworkInference/Source/NeuralNetworkInferenceQA/Private/ModelUnitTester.h
@@ -16,7 +16,7 @@ private:
 	 */
 	static bool ModelLoadAccuracyAndSpeedTests(const FString& InProjectContentDir, const FString& InModelZooRelativeDirectory, const TArray<FString>& InModelNames,
 		const TArray<float>& InInputArrayValues, const TArray<TArray<double>>& InCPUGroundTruths, const TArray<TArray<double>>& InGPUGroundTruths,
-		const TArray<int32>& InCPURepetitionsForUEAndORT, const TArray<int32>& InGPURepetitionsForUEAndORT, const TArray<int32>& InCPURepetitionsForUEOnly, const TArray<int32>& InGPURepetitionsForUEOnly);
+		const TArray<int32>& InCPURepetitionsForUEAndORT, const TArray<int32>& InGPURepetitionsForUEAndORTBackEnd, const TArray<int32>& InCPURepetitionsForUEOnly, const TArray<int32>& InGPURepetitionsForUEOnly);
 	/**
 	 * Other auxiliary functions for GlobalTest().
 	 */
@@ -25,9 +25,13 @@ private:
 	static FString GetUAssetModelFilePath(const FString& InModelName, const FString& InModelZooRelativeDirectory);
 	static UNeuralNetwork* NetworkUassetLoadTest(const FString& InUAssetPath);
 	static UNeuralNetwork* NetworkONNXOrORTLoadTest(const FString& InModelFilePath);
-	static bool ModelAccuracyTest(UNeuralNetwork* InOutNetwork, const TArray<float>& InInputArrayValues, const TArray<double>& InCPUGroundTruths,
-		const TArray<double>& InGPUGroundTruths);
+	static bool ModelAccuracyTest(UNeuralNetwork* InOutNetwork, const ENeuralBackEnd InBackEnd, const TArray<float>& InInputArrayValues, const TArray<double>& InCPUGroundTruths, const TArray<double>& InGPUGroundTruths);
 	static bool ModelSpeedTest(const FString& InUAssetPath, const ENeuralDeviceType InDeviceType, const ENeuralBackEnd InBackEnd, const int32 InRepetitions);
 	static double GetAveragedL1Norm(const TArray<float>& InArray);
 	static double GetAveragedL1NormDiff(const TArray<float>& InArray1, const TArray<float>& InArray2);
+	/**
+	 * Verbose string auxiliary functions.
+	 */
+	static FString GetDeviceTypeString(const ENeuralDeviceType InDeviceType);
+	static FString GetBackEndString(const ENeuralBackEnd InBackEnd);
 };
