NeuralNetworkInference (NNI) plugin:
- Warning returns false
- Disabled SetOutputDevice block in Run().

#rnx

diff --git a/Engine/Plugins/Experimental/NeuralNetworkInference/Source/NeuralNetworkInference/Private/NeuralNetworkImplBackEndUEAndORT.cpp b/Engine/Plugins/Experimental/NeuralNetworkInference/Source/NeuralNetworkInference/Private/NeuralNetworkImplBackEndUEAndORT.cpp
--- a/Engine/Plugins/Experimental/NeuralNetworkInference/Source/NeuralNetworkInference/Private/NeuralNetworkImplBackEndUEAndORT.cpp
+++ b/Engine/Plugins/Experimental/NeuralNetworkInference/Source/NeuralNetworkInference/Private/NeuralNetworkImplBackEndUEAndORT.cpp
@@ -213,6 +213,7 @@ bool UNeuralNetwork::FImplBackEndUEAndORT::Load(TSharedPtr<FImplBackEndUEAndORT>
 			if (InOutImplBackEndUEAndORT->DmlGPUAllocator.IsValid() && !InOutImplBackEndUEAndORT->DmlGPUAllocator->IsValid())
 			{
 				UE_LOG(LogNeuralNetworkInference, Warning, TEXT("FImplBackEndUEAndORT::Load() DirectML GPU resource allocator has failed to initialize."));
+				return false;
 			}
 #endif
 		}
@@ -261,9 +262,9 @@ void UNeuralNetwork::FImplBackEndUEAndORT::Run(const ENeuralNetworkSynchronousMo
 		const FRedirectCoutAndCerrToUeLog RedirectCoutAndCerrToUeLog;
 
 		// @todo: Temporarily disabled until we connect GPU input/output between UE and ORT
-		if (InInputDeviceType == ENeuralDeviceType::GPU || InOutputDeviceType == ENeuralDeviceType::GPU)
+		if (InInputDeviceType == ENeuralDeviceType::GPU)
 		{
-			UE_LOG(LogNeuralNetworkInference, Warning, TEXT("FImplBackEndUEAndORT::Run(): InputDeviceType & OutputDeviceType must be set to CPU for now."));
+			UE_LOG(LogNeuralNetworkInference, Warning, TEXT("FImplBackEndUEAndORT::Run(): InputDeviceType must be set to CPU for now."));
 			return;
 		}
 
