NeuralNetworkInference (NNI) plugin:
- FModelUnitTester::GlobalTest(): Modified thresholds to be more strict with CPU vs. GPU results (5e-6 rather than 1e-3)

#rnx

diff --git a/Engine/Plugins/Experimental/NeuralNetworkInference/Source/NeuralNetworkInferenceQA/Private/ModelUnitTester.cpp b/Engine/Plugins/Experimental/NeuralNetworkInference/Source/NeuralNetworkInferenceQA/Private/ModelUnitTester.cpp
--- a/Engine/Plugins/Experimental/NeuralNetworkInference/Source/NeuralNetworkInferenceQA/Private/ModelUnitTester.cpp
+++ b/Engine/Plugins/Experimental/NeuralNetworkInference/Source/NeuralNetworkInferenceQA/Private/ModelUnitTester.cpp
@@ -237,36 +237,36 @@ bool FModelUnitTester::ModelAccuracyTest(UNeuralNetwork* InOutNetwork, const ENe
 		// Input CPU + GPU + Output GPU
 		InOutNetwork->SetDeviceType(ENeuralDeviceType::GPU);
 		InOutNetwork->SetOutputDeviceType(ENeuralDeviceType::GPU);
-		//InOutNetwork->Run(); // @todo: Uncomment line to re-enable test
+		InOutNetwork->Run();
 		TArray<float> CPUGPUGPUOutput = CPUGPUCPUOutput; // InOutNetwork->GetOutputTensor().GetArrayCopy<float>(); // @todo: Modify line to re-enable test
 		// Final verbose
 		const double CPUAvgL1Norm = GetAveragedL1Norm(CPUOutput);
 		const double CPUGPUCPUAvgL1Norm = GetAveragedL1Norm(CPUGPUCPUOutput);
 		const double CPUGPUGPUAvgL1Norm = GetAveragedL1Norm(CPUGPUGPUOutput);
 		const double RelativeCoefficient = 1. / FMath::Max(1., FMath::Min(CPUAvgL1Norm, CPUGPUCPUAvgL1Norm)); // Max(1, X) to avoid 0s
-		const double CPUGPUAvgL1NormDiff = GetAveragedL1NormDiff(CPUOutput, CPUGPUCPUOutput) * RelativeCoefficient * 1e3;
-		const double GPUGPUAvgL1NormDiff = GetAveragedL1NormDiff(CPUGPUCPUOutput, CPUGPUGPUOutput) * RelativeCoefficient * 1e3;
+		const double CPUGPUAvgL1NormDiff = GetAveragedL1NormDiff(CPUOutput, CPUGPUCPUOutput) * RelativeCoefficient * 1e6;
+		const double GPUGPUAvgL1NormDiff = GetAveragedL1NormDiff(CPUGPUCPUOutput, CPUGPUGPUOutput) / FMath::Max(1., FMath::Min(CPUGPUGPUAvgL1Norm, CPUGPUCPUAvgL1Norm)) * 1e12;
 		const double FastCPUGPUAvgL1NormDiff = FMath::Abs((CPUAvgL1Norm - CPUGPUCPUAvgL1Norm)) * RelativeCoefficient * 1e6;
 		const double FastCPUAvgL1NormDiff = FMath::Abs(CPUAvgL1Norm - CPUGroundTruth) / FMath::Max(1., FMath::Min(CPUAvgL1Norm, CPUGroundTruth)) * 1e7;
 		const double FastGPUAvgL1NormDiff = FMath::Abs(CPUGPUCPUAvgL1Norm - GPUGroundTruth) / FMath::Max(1., FMath::Min(CPUGPUCPUAvgL1Norm, GPUGroundTruth)) * 1e7;
 		UE_LOG(LogNeuralNetworkInferenceQA, Display, TEXT("%s: InputNorm = %f, OutputNormCPU = %f, OutputNormGPU = %f, OutputNormCPUGPUGPU = %f, OutputNormGT = %f, CPUAvgL1Norm = %f, CPUGPUCPUAvgL1Norm = %f,"),
 			*BackEndString, GetAveragedL1Norm(InputArray), CPUAvgL1Norm, CPUGPUCPUAvgL1Norm, CPUGPUGPUAvgL1Norm, CPUGroundTruth, CPUAvgL1Norm, CPUGPUCPUAvgL1Norm);
-		UE_LOG(LogNeuralNetworkInferenceQA, Display, TEXT("\tCPUGPUAvgL1NormDiff = %fe-3, GPUGPUAvgL1NormDiff = %fe-3, FastCPUGPUAvgL1NormDiff = %fe-6, FastCPUAvgL1NormDiff = %fe-7, FastGPUAvgL1NormDiff = %fe-7 (1e-7 is roughly the precision for float)."),
+		UE_LOG(LogNeuralNetworkInferenceQA, Display, TEXT("\tCPUGPUAvgL1NormDiff = %fe-6, GPUGPUAvgL1NormDiff = %fe-7, FastCPUGPUAvgL1NormDiff = %fe-6, FastCPUAvgL1NormDiff = %fe-7, FastGPUAvgL1NormDiff = %fe-7 (1e-7 is roughly the precision for float)."),
 			CPUGPUAvgL1NormDiff, GPUGPUAvgL1NormDiff, FastCPUGPUAvgL1NormDiff, FastCPUAvgL1NormDiff, FastGPUAvgL1NormDiff);
 		const TArray<int64>& InputSizes = InOutNetwork->GetInputTensor().GetSizes();
 		const TArray<int64>& OutputSizes = InOutNetwork->GetOutputTensor().GetSizes();
 		const int64 MaxNumberElementsToDisplay = 100;
 		// Check if test failed and (if so) display information
 		bool bDidSomeTestFailed = (!FMath::IsFinite(FastCPUGPUAvgL1NormDiff) || FastCPUGPUAvgL1NormDiff > 5);
-		bDidSomeTestFailed |= (!FMath::IsFinite(CPUGPUAvgL1NormDiff) || CPUGPUAvgL1NormDiff > 1);
+		bDidSomeTestFailed |= (!FMath::IsFinite(CPUGPUAvgL1NormDiff) || CPUGPUAvgL1NormDiff > 5);
 		bDidSomeTestFailed |= (!FMath::IsFinite(GPUGPUAvgL1NormDiff) || GPUGPUAvgL1NormDiff > 1);
 		bDidSomeTestFailed |= (!FMath::IsFinite(FastCPUAvgL1NormDiff) || FastCPUAvgL1NormDiff > 30);
 		bDidSomeTestFailed |= (!FMath::IsFinite(FastGPUAvgL1NormDiff) || FastGPUAvgL1NormDiff > 30);
 		if (bDidSomeTestFailed)
 		{
 			UE_LOG(LogNeuralNetworkInferenceQA, Display, TEXT("FastCPUGPUAvgL1NormDiff (%fe-6) < 5e-6 might have failed."), FastCPUGPUAvgL1NormDiff);
-			UE_LOG(LogNeuralNetworkInferenceQA, Display, TEXT("CPUGPUAvgL1NormDiff (%fe-7) < 1e-7 might have failed."), CPUGPUAvgL1NormDiff);
-			UE_LOG(LogNeuralNetworkInferenceQA, Display, TEXT("GPUGPUAvgL1NormDiff (%fe-7) < 1e-7 might have failed."), GPUGPUAvgL1NormDiff);
+			UE_LOG(LogNeuralNetworkInferenceQA, Display, TEXT("CPUGPUAvgL1NormDiff (%fe-6) < 5e-6 might have failed."), CPUGPUAvgL1NormDiff);
+			UE_LOG(LogNeuralNetworkInferenceQA, Display, TEXT("GPUGPUAvgL1NormDiff (%fe-12) < 1e-12 might have failed."), GPUGPUAvgL1NormDiff);
 			UE_LOG(LogNeuralNetworkInferenceQA, Display, TEXT("FastCPUAvgL1NormDiff (%fe-7) < 30e-7 might have failed (~30 times the float precision).\nCPUOutput = %s."), FastCPUAvgL1NormDiff, *FNeuralTensor(CPUOutput, OutputSizes).ToString(MaxNumberElementsToDisplay));
 			UE_LOG(LogNeuralNetworkInferenceQA, Display, TEXT("FastGPUAvgL1NormDiff (%fe-7) < 30e-7 might have failed (~30 times the float precision).\nCPUGPUCPUOutput = %s."), FastGPUAvgL1NormDiff, *FNeuralTensor(CPUGPUCPUOutput, OutputSizes).ToString(MaxNumberElementsToDisplay));
 			UE_LOG(LogNeuralNetworkInferenceQA, Display, TEXT("Input = %s"),
