Quick code clean up for future mini refactor.

#rnx

diff --git a/Engine/Plugins/Experimental/NeuralNetworkInference/Source/NeuralNetworkInferenceQA/Private/ModelUnitTester.cpp b/Engine/Plugins/Experimental/NeuralNetworkInference/Source/NeuralNetworkInferenceQA/Private/ModelUnitTester.cpp
--- a/Engine/Plugins/Experimental/NeuralNetworkInference/Source/NeuralNetworkInferenceQA/Private/ModelUnitTester.cpp
+++ b/Engine/Plugins/Experimental/NeuralNetworkInference/Source/NeuralNetworkInferenceQA/Private/ModelUnitTester.cpp
@@ -226,19 +226,24 @@ bool FModelUnitTester::ModelAccuracyTest(UNeuralNetwork* InOutNetwork, const ENe
 		const double CPUGroundTruth = InCPUGroundTruths[Index];
 		const double GPUGroundTruth = InGPUGroundTruths[Index];
 		InOutNetwork->SetInputFromArrayCopy(InputArray);
+		//// Input CPU + GPU + Output CPU
+		//InOutNetwork->SetDeviceType(ENeuralDeviceType::GPU);
+		//InOutNetwork->Run();
+		//const TArray<float> CPUGPUCPUOutput = InOutNetwork->GetOutputTensor().GetArrayCopy<float>();
 		// CPU
 		InOutNetwork->SetDeviceType(ENeuralDeviceType::CPU);
 		InOutNetwork->Run();
-		TArray<float> CPUOutput = InOutNetwork->GetOutputTensor().GetArrayCopy<float>();
-		// Input CPU + GPU + Output CPU
-		InOutNetwork->SetDeviceType(ENeuralDeviceType::GPU);
-		InOutNetwork->Run();
-		TArray<float> CPUGPUCPUOutput = InOutNetwork->GetOutputTensor().GetArrayCopy<float>();
+		const TArray<float> CPUOutput = InOutNetwork->GetOutputTensor().GetArrayCopy<float>();
+// Input CPU + GPU + Output CPU
+InOutNetwork->SetDeviceType(ENeuralDeviceType::GPU);
+InOutNetwork->Run();
+const TArray<float> CPUGPUCPUOutput = InOutNetwork->GetOutputTensor().GetArrayCopy<float>();
 		// Input CPU + GPU + Output GPU
 		InOutNetwork->SetDeviceType(ENeuralDeviceType::GPU);
 		InOutNetwork->SetOutputDeviceType(ENeuralDeviceType::GPU);
-		InOutNetwork->Run();
-		TArray<float> CPUGPUGPUOutput = CPUGPUCPUOutput; // InOutNetwork->GetOutputTensor().GetArrayCopy<float>(); // @todo: Modify line to re-enable test
+		//InOutNetwork->Run();
+		//InOutNetwork->OutputTensorsToCPU();
+		const TArray<float> CPUGPUGPUOutput = CPUGPUCPUOutput; // InOutNetwork->GetOutputTensor().GetArrayCopy<float>(); // @todo: Modify line to re-enable test
 		// Final verbose
 		const double CPUAvgL1Norm = GetAveragedL1Norm(CPUOutput);
 		const double CPUGPUCPUAvgL1Norm = GetAveragedL1Norm(CPUGPUCPUOutput);
@@ -257,11 +262,11 @@ bool FModelUnitTester::ModelAccuracyTest(UNeuralNetwork* InOutNetwork, const ENe
 		const TArray<int64>& OutputSizes = InOutNetwork->GetOutputTensor().GetSizes();
 		const int64 MaxNumberElementsToDisplay = 100;
 		// Check if test failed and (if so) display information
-		bool bDidSomeTestFailed = (!FMath::IsFinite(FastCPUGPUAvgL1NormDiff) || FastCPUGPUAvgL1NormDiff > 5);
-		bDidSomeTestFailed |= (!FMath::IsFinite(CPUGPUAvgL1NormDiff) || CPUGPUAvgL1NormDiff > 5);
-		bDidSomeTestFailed |= (!FMath::IsFinite(GPUGPUAvgL1NormDiff) || GPUGPUAvgL1NormDiff > 1);
-		bDidSomeTestFailed |= (!FMath::IsFinite(FastCPUAvgL1NormDiff) || FastCPUAvgL1NormDiff > 30);
-		bDidSomeTestFailed |= (!FMath::IsFinite(FastGPUAvgL1NormDiff) || FastGPUAvgL1NormDiff > 30);
+		const bool bDidSomeTestFailed = (!FMath::IsFinite(FastCPUGPUAvgL1NormDiff) || FastCPUGPUAvgL1NormDiff > 5)
+			|| (!FMath::IsFinite(CPUGPUAvgL1NormDiff) || CPUGPUAvgL1NormDiff > 5)
+			|| (!FMath::IsFinite(GPUGPUAvgL1NormDiff) || GPUGPUAvgL1NormDiff > 1)
+			|| (!FMath::IsFinite(FastCPUAvgL1NormDiff) || FastCPUAvgL1NormDiff > 30)
+			|| (!FMath::IsFinite(FastGPUAvgL1NormDiff) || FastGPUAvgL1NormDiff > 30);
 		if (bDidSomeTestFailed)
 		{
 			UE_LOG(LogNeuralNetworkInferenceQA, Display, TEXT("FastCPUGPUAvgL1NormDiff (%fe-6) < 5e-6 might have failed."), FastCPUGPUAvgL1NormDiff);
