NeuralNetworkInference (NNI) plugin:
- FModelUnitTester::GlobalTest() code refactor to avoid cpu-gpu bugs.

#rnx

diff --git a/Engine/Plugins/Experimental/NeuralNetworkInference/Source/NeuralNetworkInference/Private/NeuralNetwork.cpp b/Engine/Plugins/Experimental/NeuralNetworkInference/Source/NeuralNetworkInference/Private/NeuralNetwork.cpp
--- a/Engine/Plugins/Experimental/NeuralNetworkInference/Source/NeuralNetworkInference/Private/NeuralNetwork.cpp
+++ b/Engine/Plugins/Experimental/NeuralNetworkInference/Source/NeuralNetworkInference/Private/NeuralNetwork.cpp
@@ -8,6 +8,9 @@
 #include "HAL/FileManager.h"
 #include "Misc/FileHelper.h"
 #include "Misc/Paths.h"
+#include "RenderGraphBuilder.h"
+#include "RenderingThread.h"
+#include "RHI.h"
 
 
 
@@ -342,10 +345,51 @@ int64 UNeuralNetwork::GetOutputTensorNumber() const
 	}
 
 	// Unknown
-	checkf(false, TEXT("UNeuralNetwork::GetOutputTensorNumber(): Unknown [BackEnd,BackEndForCurrentPlatform] = [%d,%d]."), (int32)BackEnd, (int32)BackEndForCurrentPlatform);
+	UE_LOG(LogNeuralNetworkInference, Warning, TEXT("UNeuralNetwork::GetOutputTensorNumber(): Unknown [BackEnd,BackEndForCurrentPlatform] = [%d,%d]."), (int32)BackEnd, (int32)BackEndForCurrentPlatform);
 	return -1;
 }
 
+void UNeuralNetwork::OutputTensorsToCPU(const TArray<int32>& InOutputTensorIndexes)
+{
+	// Sanity check
+	if (!bIsLoaded)
+	{
+		UE_LOG(LogNeuralNetworkInference, Warning, TEXT("UNeuralNetwork::OutputTensorToCPU(): Call UNeuralNetwork::Load() to load a model first."));
+		return;
+	}
+
+	// In RHI
+	ENQUEUE_RENDER_COMMAND(UNeuralNetwork_OutputTensorToCPU_RenderThread)(
+		[this, &InOutputTensorIndexes](FRHICommandListImmediate& RHICmdList)
+		{
+			FRDGBuilder GraphBuilder(RHICmdList, RDG_EVENT_NAME("UNeuralNetwork::OutputTensorToCPU()"));
+			// Run for each output tensor
+			if (InOutputTensorIndexes.IsEmpty())
+			{
+				// Refresh tensor(s) w.r.t. GraphBuilder + Move memory from GPU to CPU
+				for (uint32 OutputTensorIndex = 0; OutputTensorIndex < GetOutputTensorNumber(); ++OutputTensorIndex)
+				{
+					FNeuralTensor OutputTensor = GetOutputTensor(OutputTensorIndex);
+					OutputTensor.UpdateSRVAndOrUAV_RenderThread(&GraphBuilder);
+					OutputTensor.ToCPU_RenderThread(&GraphBuilder);
+				}
+			}
+			// Run for the desired output tensors
+			else
+			{
+				for (const int32 OutputTensorIndex : InOutputTensorIndexes)
+				{
+					FNeuralTensor OutputTensor = GetOutputTensor(OutputTensorIndex);
+					OutputTensor.UpdateSRVAndOrUAV_RenderThread(&GraphBuilder);
+					OutputTensor.ToCPU_RenderThread(&GraphBuilder);
+				}
+			}
+			// Execute Render Graph
+			GraphBuilder.Execute();
+		}
+	);
+}
+
 void UNeuralNetwork::Run()
 {
 	DECLARE_SCOPE_CYCLE_COUNTER(TEXT("UNeuralNetwork_Run"), STAT_UNeuralNetwork_Run, STATGROUP_MachineLearning);
diff --git a/Engine/Plugins/Experimental/NeuralNetworkInference/Source/NeuralNetworkInference/Public/NeuralNetwork.h b/Engine/Plugins/Experimental/NeuralNetworkInference/Source/NeuralNetworkInference/Public/NeuralNetwork.h
--- a/Engine/Plugins/Experimental/NeuralNetworkInference/Source/NeuralNetworkInference/Public/NeuralNetwork.h
+++ b/Engine/Plugins/Experimental/NeuralNetworkInference/Source/NeuralNetworkInference/Public/NeuralNetwork.h
@@ -144,6 +144,12 @@ public:
 	 */
 	const FNeuralTensor& GetOutputTensor(const int32 InTensorIndex = 0) const;
 	int64 GetOutputTensorNumber() const;
+	/**
+	 * Non-efficient function meant to be used only for debugging purposes. It will send the GPU memory of the desired output tensor(s) back into CPU memory.
+	 * It is useful for debugging OutputDeviceType == ENeuralDeviceType::GPU cases.
+	 * @param InOutputTensorIndexes If empty (default value), it will apply to all output tensors.
+	 */
+	void OutputTensorsToCPU(const TArray<int32>& InOutputTensorIndexes = TArray<int32>());
 
 	/**
 	 * Run() executes the forward pass on the current UNeuralNetwork given the current input FDeprecatedNeuralTensor(s), which were previously filled with
diff --git a/Engine/Plugins/Experimental/NeuralNetworkInference/Source/NeuralNetworkInferenceQA/Private/ModelUnitTester.cpp b/Engine/Plugins/Experimental/NeuralNetworkInference/Source/NeuralNetworkInferenceQA/Private/ModelUnitTester.cpp
--- a/Engine/Plugins/Experimental/NeuralNetworkInference/Source/NeuralNetworkInferenceQA/Private/ModelUnitTester.cpp
+++ b/Engine/Plugins/Experimental/NeuralNetworkInference/Source/NeuralNetworkInferenceQA/Private/ModelUnitTester.cpp
@@ -220,47 +220,48 @@ bool FModelUnitTester::ModelAccuracyTest(UNeuralNetwork* InOutNetwork, const ENe
 	const FString BackEndString = GetBackEndString(InBackEnd);
 	// Run each input with CPU/GPU and compare with each other and with the ground truth
 	bool bDidGlobalTestPassed = true;
+	TArray<TArray<float>> CPUOutputs, CPUGPUCPUOutputs, CPUGPUGPUOutputs;
 	for (int32 Index = 0; Index < InputArrays.Num(); ++Index)
 	{
-		TArray<float>& InputArray = InputArrays[Index];
-		const double CPUGroundTruth = InCPUGroundTruths[Index];
-		const double GPUGroundTruth = InGPUGroundTruths[Index];
-		InOutNetwork->SetInputFromArrayCopy(InputArray);
+		InOutNetwork->SetInputFromArrayCopy(InputArrays[Index]);
 		//// Input CPU + GPU + Output CPU
 		//InOutNetwork->SetDeviceType(ENeuralDeviceType::GPU);
 		//InOutNetwork->Run();
-		//const TArray<float> CPUGPUCPUOutput = InOutNetwork->GetOutputTensor().GetArrayCopy<float>();
+		//CPUGPUCPUOutputs.Emplace(InOutNetwork->GetOutputTensor().GetArrayCopy<float>());
 		// CPU
 		InOutNetwork->SetDeviceType(ENeuralDeviceType::CPU);
 		InOutNetwork->Run();
-		const TArray<float> CPUOutput = InOutNetwork->GetOutputTensor().GetArrayCopy<float>();
+		CPUOutputs.Emplace(InOutNetwork->GetOutputTensor().GetArrayCopy<float>());
 // Input CPU + GPU + Output CPU
 InOutNetwork->SetDeviceType(ENeuralDeviceType::GPU);
 InOutNetwork->Run();
-const TArray<float> CPUGPUCPUOutput = InOutNetwork->GetOutputTensor().GetArrayCopy<float>();
+CPUGPUCPUOutputs.Emplace(InOutNetwork->GetOutputTensor().GetArrayCopy<float>());
 		// Input CPU + GPU + Output GPU
 		InOutNetwork->SetDeviceType(ENeuralDeviceType::GPU);
 		InOutNetwork->SetOutputDeviceType(ENeuralDeviceType::GPU);
 		//InOutNetwork->Run();
 		//InOutNetwork->OutputTensorsToCPU();
-		const TArray<float> CPUGPUGPUOutput = CPUGPUCPUOutput; // InOutNetwork->GetOutputTensor().GetArrayCopy<float>(); // @todo: Modify line to re-enable test
-		// Final verbose
-		const double CPUAvgL1Norm = GetAveragedL1Norm(CPUOutput);
-		const double CPUGPUCPUAvgL1Norm = GetAveragedL1Norm(CPUGPUCPUOutput);
-		const double CPUGPUGPUAvgL1Norm = GetAveragedL1Norm(CPUGPUGPUOutput);
+CPUGPUGPUOutputs.Push(CPUGPUCPUOutputs[Index]); // @todo: Modify line to re-enable test
+		//CPUGPUGPUOutputs.Emplace(InOutNetwork->GetOutputTensor().GetArrayCopy<float>()); // @todo: Modify line to re-enable test
+	}
+	// The second for loop is to make sure running on the CPU does not affect GPU results or vice versa
+	for (int32 Index = 0; Index < InputArrays.Num(); ++Index)
+	{
+		// Prepare verbose
+		const double CPUAvgL1Norm = GetAveragedL1Norm(CPUOutputs[Index]);
+		const double CPUGPUCPUAvgL1Norm = GetAveragedL1Norm(CPUGPUCPUOutputs[Index]);
+		const double CPUGPUGPUAvgL1Norm = GetAveragedL1Norm(CPUGPUGPUOutputs[Index]);
 		const double RelativeCoefficient = 1. / FMath::Max(1., FMath::Min(CPUAvgL1Norm, CPUGPUCPUAvgL1Norm)); // Max(1, X) to avoid 0s
-		const double CPUGPUAvgL1NormDiff = GetAveragedL1NormDiff(CPUOutput, CPUGPUCPUOutput) * RelativeCoefficient * 1e6;
-		const double GPUGPUAvgL1NormDiff = GetAveragedL1NormDiff(CPUGPUCPUOutput, CPUGPUGPUOutput) / FMath::Max(1., FMath::Min(CPUGPUGPUAvgL1Norm, CPUGPUCPUAvgL1Norm)) * 1e12;
+		const double CPUGPUAvgL1NormDiff = GetAveragedL1NormDiff(CPUOutputs[Index], CPUGPUCPUOutputs[Index]) * RelativeCoefficient * 1e6;
+		const double GPUGPUAvgL1NormDiff = GetAveragedL1NormDiff(CPUGPUCPUOutputs[Index], CPUGPUGPUOutputs[Index]) / FMath::Max(1., FMath::Min(CPUGPUGPUAvgL1Norm, CPUGPUCPUAvgL1Norm)) * 1e12;
 		const double FastCPUGPUAvgL1NormDiff = FMath::Abs((CPUAvgL1Norm - CPUGPUCPUAvgL1Norm)) * RelativeCoefficient * 1e6;
-		const double FastCPUAvgL1NormDiff = FMath::Abs(CPUAvgL1Norm - CPUGroundTruth) / FMath::Max(1., FMath::Min(CPUAvgL1Norm, CPUGroundTruth)) * 1e7;
-		const double FastGPUAvgL1NormDiff = FMath::Abs(CPUGPUCPUAvgL1Norm - GPUGroundTruth) / FMath::Max(1., FMath::Min(CPUGPUCPUAvgL1Norm, GPUGroundTruth)) * 1e7;
+		const double FastCPUAvgL1NormDiff = FMath::Abs(CPUAvgL1Norm - InCPUGroundTruths[Index]) / FMath::Max(1., FMath::Min(CPUAvgL1Norm, InCPUGroundTruths[Index])) * 1e7;
+		const double FastGPUAvgL1NormDiff = FMath::Abs(CPUGPUCPUAvgL1Norm - InGPUGroundTruths[Index]) / FMath::Max(1., FMath::Min(CPUGPUCPUAvgL1Norm, InGPUGroundTruths[Index])) * 1e7;
+		// Print verbose
 		UE_LOG(LogNeuralNetworkInferenceQA, Display, TEXT("%s: InputNorm = %f, OutputNormCPU = %f, OutputNormGPU = %f, OutputNormCPUGPUGPU = %f, OutputNormGT = %f, CPUAvgL1Norm = %f, CPUGPUCPUAvgL1Norm = %f,"),
-			*BackEndString, GetAveragedL1Norm(InputArray), CPUAvgL1Norm, CPUGPUCPUAvgL1Norm, CPUGPUGPUAvgL1Norm, CPUGroundTruth, CPUAvgL1Norm, CPUGPUCPUAvgL1Norm);
+			*BackEndString, GetAveragedL1Norm(InputArrays[Index]), CPUAvgL1Norm, CPUGPUCPUAvgL1Norm, CPUGPUGPUAvgL1Norm, InCPUGroundTruths[Index], CPUAvgL1Norm, CPUGPUCPUAvgL1Norm);
 		UE_LOG(LogNeuralNetworkInferenceQA, Display, TEXT("\tCPUGPUAvgL1NormDiff = %fe-6, GPUGPUAvgL1NormDiff = %fe-7, FastCPUGPUAvgL1NormDiff = %fe-6, FastCPUAvgL1NormDiff = %fe-7, FastGPUAvgL1NormDiff = %fe-7 (1e-7 is roughly the precision for float)."),
 			CPUGPUAvgL1NormDiff, GPUGPUAvgL1NormDiff, FastCPUGPUAvgL1NormDiff, FastCPUAvgL1NormDiff, FastGPUAvgL1NormDiff);
-		const TArray<int64>& InputSizes = InOutNetwork->GetInputTensor().GetSizes();
-		const TArray<int64>& OutputSizes = InOutNetwork->GetOutputTensor().GetSizes();
-		const int64 MaxNumberElementsToDisplay = 100;
 		// Check if test failed and (if so) display information
 		const bool bDidSomeTestFailed = (!FMath::IsFinite(FastCPUGPUAvgL1NormDiff) || FastCPUGPUAvgL1NormDiff > 5)
 			|| (!FMath::IsFinite(CPUGPUAvgL1NormDiff) || CPUGPUAvgL1NormDiff > 5)
@@ -269,15 +270,19 @@ const TArray<float> CPUGPUCPUOutput = InOutNetwork->GetOutputTensor().GetArrayCo
 			|| (!FMath::IsFinite(FastGPUAvgL1NormDiff) || FastGPUAvgL1NormDiff > 30);
 		if (bDidSomeTestFailed)
 		{
+			const int64 MaxNumberElementsToDisplay = 100;
+			const TArray<int64>& InputSizes = InOutNetwork->GetInputTensor().GetSizes();
+			const TArray<int64>& OutputSizes = InOutNetwork->GetOutputTensor().GetSizes();
 			UE_LOG(LogNeuralNetworkInferenceQA, Display, TEXT("FastCPUGPUAvgL1NormDiff (%fe-6) < 5e-6 might have failed."), FastCPUGPUAvgL1NormDiff);
 			UE_LOG(LogNeuralNetworkInferenceQA, Display, TEXT("CPUGPUAvgL1NormDiff (%fe-6) < 5e-6 might have failed."), CPUGPUAvgL1NormDiff);
 			UE_LOG(LogNeuralNetworkInferenceQA, Display, TEXT("GPUGPUAvgL1NormDiff (%fe-12) < 1e-12 might have failed."), GPUGPUAvgL1NormDiff);
-			UE_LOG(LogNeuralNetworkInferenceQA, Display, TEXT("FastCPUAvgL1NormDiff (%fe-7) < 30e-7 might have failed (~30 times the float precision).\nCPUOutput = %s."), FastCPUAvgL1NormDiff, *FNeuralTensor(CPUOutput, OutputSizes).ToString(MaxNumberElementsToDisplay));
-			UE_LOG(LogNeuralNetworkInferenceQA, Display, TEXT("FastGPUAvgL1NormDiff (%fe-7) < 30e-7 might have failed (~30 times the float precision).\nCPUGPUCPUOutput = %s."), FastGPUAvgL1NormDiff, *FNeuralTensor(CPUGPUCPUOutput, OutputSizes).ToString(MaxNumberElementsToDisplay));
+			UE_LOG(LogNeuralNetworkInferenceQA, Display, TEXT("FastCPUAvgL1NormDiff (%fe-7) < 30e-7 might have failed (~30 times the float precision).\nCPUOutput = %s."), FastCPUAvgL1NormDiff, *FNeuralTensor(CPUOutputs[Index], OutputSizes).ToString(MaxNumberElementsToDisplay));
+			UE_LOG(LogNeuralNetworkInferenceQA, Display, TEXT("FastGPUAvgL1NormDiff (%fe-7) < 30e-7 might have failed (~30 times the float precision).\nCPUGPUCPUOutput = %s."), FastGPUAvgL1NormDiff, *FNeuralTensor(CPUGPUCPUOutputs[Index], OutputSizes).ToString(MaxNumberElementsToDisplay));
 			UE_LOG(LogNeuralNetworkInferenceQA, Display, TEXT("Input = %s"),
 				*FNeuralTensor(InOutNetwork->GetInputTensor().GetArrayCopy<float>(), InputSizes).ToString(MaxNumberElementsToDisplay));
-			UE_LOG(LogNeuralNetworkInferenceQA, Display, TEXT("CPUOutput = %s"), *FNeuralTensor(CPUOutput, OutputSizes).ToString(MaxNumberElementsToDisplay));
-			UE_LOG(LogNeuralNetworkInferenceQA, Display, TEXT("CPUGPUCPUOutput = %s"), *FNeuralTensor(CPUGPUCPUOutput, OutputSizes).ToString(MaxNumberElementsToDisplay));
+			UE_LOG(LogNeuralNetworkInferenceQA, Display, TEXT("CPUOutput = %s"), *FNeuralTensor(CPUOutputs[Index], OutputSizes).ToString(MaxNumberElementsToDisplay));
+			UE_LOG(LogNeuralNetworkInferenceQA, Display, TEXT("CPUGPUCPUOutput = %s"), *FNeuralTensor(CPUGPUCPUOutputs[Index], OutputSizes).ToString(MaxNumberElementsToDisplay));
+			UE_LOG(LogNeuralNetworkInferenceQA, Display, TEXT("CPUGPUGPUOutput = %s"), *FNeuralTensor(CPUGPUGPUOutputs[Index], OutputSizes).ToString(MaxNumberElementsToDisplay));
 			UE_LOG(LogNeuralNetworkInferenceQA, Warning, TEXT("At least 1 of the 5 CPU/GPU tests failed."));
 			return false;
 		}
