Fixed bug in ModelUnitTester.cpp: SetOutputDevice(CPU) when no GPU.
#rnx

diff --git a/Engine/Plugins/Experimental/NeuralNetworkInference/Source/NeuralNetworkInferenceQA/Private/ModelUnitTester.cpp b/Engine/Plugins/Experimental/NeuralNetworkInference/Source/NeuralNetworkInferenceQA/Private/ModelUnitTester.cpp
--- a/Engine/Plugins/Experimental/NeuralNetworkInference/Source/NeuralNetworkInferenceQA/Private/ModelUnitTester.cpp
+++ b/Engine/Plugins/Experimental/NeuralNetworkInference/Source/NeuralNetworkInferenceQA/Private/ModelUnitTester.cpp
@@ -226,19 +226,26 @@ bool FModelUnitTester::ModelAccuracyTest(UNeuralNetwork* InOutNetwork, const ENe
 		InOutNetwork->SetInputFromArrayCopy(InputArrays[Index]);
 		//// Input CPU + GPU + Output CPU
 		//InOutNetwork->SetDeviceType(ENeuralDeviceType::GPU);
+		//InOutNetwork->SetInputDeviceType(ENeuralDeviceType::CPU);
+		//InOutNetwork->SetOutputDeviceType(ENeuralDeviceType::CPU);
 		//InOutNetwork->Run();
 		//CPUGPUCPUOutputs.Emplace(InOutNetwork->GetOutputTensor().GetArrayCopy<float>());
 		// CPU
 		InOutNetwork->SetDeviceType(ENeuralDeviceType::CPU);
+		InOutNetwork->SetOutputDeviceType(ENeuralDeviceType::CPU);
+		InOutNetwork->SetInputDeviceType(ENeuralDeviceType::CPU);
 		InOutNetwork->Run();
 		CPUOutputs.Emplace(InOutNetwork->GetOutputTensor().GetArrayCopy<float>());
 // Input CPU + GPU + Output CPU
 InOutNetwork->SetDeviceType(ENeuralDeviceType::GPU);
+InOutNetwork->SetInputDeviceType(ENeuralDeviceType::CPU);
+InOutNetwork->SetOutputDeviceType(ENeuralDeviceType::CPU);
 InOutNetwork->Run();
 CPUGPUCPUOutputs.Emplace(InOutNetwork->GetOutputTensor().GetArrayCopy<float>());
-		// Input CPU + GPU + Output GPU
-		InOutNetwork->SetDeviceType(ENeuralDeviceType::GPU);
-		InOutNetwork->SetOutputDeviceType(ENeuralDeviceType::GPU);
+		//// Input CPU + GPU + Output GPU
+		//InOutNetwork->SetDeviceType(ENeuralDeviceType::GPU);
+		//InOutNetwork->SetInputDeviceType(ENeuralDeviceType::CPU);
+		//InOutNetwork->SetOutputDeviceType(ENeuralDeviceType::GPU);
 		//InOutNetwork->Run();
 		//InOutNetwork->OutputTensorsToCPU();
 CPUGPUGPUOutputs.Push(CPUGPUCPUOutputs[Index]); // @todo: Modify line to re-enable test
