NeuralNetworkInference (NNI) plugin:
- Fixes for debugging function UNeuralNetwork::OutputTensorsToCPU.
- Enabled SetOutputDevice(GPU) tests for UEOnly.
- Added private UNN::GetOutputTensorMutable.

#rnx

diff --git a/Engine/Plugins/Experimental/NeuralNetworkInference/Source/NeuralNetworkInference/Private/NeuralNetwork.cpp b/Engine/Plugins/Experimental/NeuralNetworkInference/Source/NeuralNetworkInference/Private/NeuralNetwork.cpp
--- a/Engine/Plugins/Experimental/NeuralNetworkInference/Source/NeuralNetworkInference/Private/NeuralNetwork.cpp
+++ b/Engine/Plugins/Experimental/NeuralNetworkInference/Source/NeuralNetworkInference/Private/NeuralNetwork.cpp
@@ -300,27 +300,32 @@ int64 UNeuralNetwork::GetInputTensorNumber() const
 	return -1;
 }
 
+#define GET_OUTPUT_TENSOR_CODE(InTensorIndex, UEOnly_GetTensorsFunction) \
+	/* Sanity check */ \
+	checkf(bIsLoaded, TEXT("Call UNeuralNetwork::Load() to load a model first.")); \
+	/* UEAndORT */ \
+	if (BackEndForCurrentPlatform == ENeuralBackEnd::UEAndORT) \
+	{ \
+		return ImplBackEndUEAndORT->OutputTensors[InTensorIndex]; \
+	} \
+	/* UEOnly */ \
+	else if (BackEndForCurrentPlatform == ENeuralBackEnd::UEOnly) \
+	{ \
+		FNeuralTensorManager& TensorManager = ImplBackEndUEOnly->TensorManager; \
+		return TensorManager.UEOnly_GetTensorsFunction()[TensorManager.GetOutputIndexes()[InTensorIndex]]; \
+	} \
+	/* Unknown */ \
+	checkf(false, TEXT("UNeuralNetwork::GetOutputTensor(): Unknown [BackEnd,BackEndForCurrentPlatform] = [%d,%d]."), (int32)BackEnd, (int32)BackEndForCurrentPlatform); \
+	return ImplBackEndUEAndORT->OutputTensors[InTensorIndex]
+
 const FNeuralTensor& UNeuralNetwork::GetOutputTensor(const int32 InTensorIndex) const
 {
-	// Sanity check
-	checkf(bIsLoaded, TEXT("Call UNeuralNetwork::Load() to load a model first."));
-
-	// UEAndORT
-	if (BackEndForCurrentPlatform == ENeuralBackEnd::UEAndORT)
-	{
-		return ImplBackEndUEAndORT->OutputTensors[InTensorIndex];
-	}
-
-	// UEOnly
-	else if (BackEndForCurrentPlatform == ENeuralBackEnd::UEOnly)
-	{
-		FNeuralTensorManager& TensorManager = ImplBackEndUEOnly->TensorManager;
-		return TensorManager.GetTensors()[TensorManager.GetOutputIndexes()[InTensorIndex]];
-	}
+	GET_OUTPUT_TENSOR_CODE(InTensorIndex, GetTensors);
+}
 
-	// Unknown
-	checkf(false, TEXT("UNeuralNetwork::GetOutputTensor(): Unknown [BackEnd,BackEndForCurrentPlatform] = [%d,%d]."), (int32)BackEnd, (int32)BackEndForCurrentPlatform);
-	return ImplBackEndUEAndORT->OutputTensors[InTensorIndex];
+FNeuralTensor& UNeuralNetwork::GetOutputTensorMutable(const int32 InTensorIndex)
+{
+	GET_OUTPUT_TENSOR_CODE(InTensorIndex, GetTensorsMutable);
 }
 
 int64 UNeuralNetwork::GetOutputTensorNumber() const
@@ -369,7 +374,7 @@ void UNeuralNetwork::OutputTensorsToCPU(const TArray<int32>& InOutputTensorIndex
 				// Refresh tensor(s) w.r.t. GraphBuilder + Move memory from GPU to CPU
 				for (uint32 OutputTensorIndex = 0; OutputTensorIndex < GetOutputTensorNumber(); ++OutputTensorIndex)
 				{
-					FNeuralTensor OutputTensor = GetOutputTensor(OutputTensorIndex);
+					FNeuralTensor& OutputTensor = GetOutputTensorMutable(OutputTensorIndex);
 					OutputTensor.UpdateSRVAndOrUAV_RenderThread(&GraphBuilder);
 					OutputTensor.ToCPU_RenderThread(&GraphBuilder);
 				}
@@ -379,7 +384,7 @@ void UNeuralNetwork::OutputTensorsToCPU(const TArray<int32>& InOutputTensorIndex
 			{
 				for (const int32 OutputTensorIndex : InOutputTensorIndexes)
 				{
-					FNeuralTensor OutputTensor = GetOutputTensor(OutputTensorIndex);
+					FNeuralTensor& OutputTensor = GetOutputTensorMutable(OutputTensorIndex);
 					OutputTensor.UpdateSRVAndOrUAV_RenderThread(&GraphBuilder);
 					OutputTensor.ToCPU_RenderThread(&GraphBuilder);
 				}
@@ -388,6 +393,8 @@ void UNeuralNetwork::OutputTensorsToCPU(const TArray<int32>& InOutputTensorIndex
 			GraphBuilder.Execute();
 		}
 	);
+
+	FNeuralNetworkInferenceUtils::WaitUntilRHIFinished();
 }
 
 void UNeuralNetwork::Run()
diff --git a/Engine/Plugins/Experimental/NeuralNetworkInference/Source/NeuralNetworkInference/Public/NeuralNetwork.h b/Engine/Plugins/Experimental/NeuralNetworkInference/Source/NeuralNetworkInference/Public/NeuralNetwork.h
--- a/Engine/Plugins/Experimental/NeuralNetworkInference/Source/NeuralNetworkInference/Public/NeuralNetwork.h
+++ b/Engine/Plugins/Experimental/NeuralNetworkInference/Source/NeuralNetworkInference/Public/NeuralNetwork.h
@@ -246,6 +246,11 @@ private:
 	 */
 	bool Load();
 
+	/**
+	 * Mutable version of GetOutputTensor(const int32 InTensorIndex) const.
+	 */
+	FNeuralTensor& GetOutputTensorMutable(const int32 InTensorIndex = 0);
+
 public:
 #if WITH_EDITOR
 	/**
diff --git a/Engine/Plugins/Experimental/NeuralNetworkInference/Source/NeuralNetworkInferenceQA/Private/ModelUnitTester.cpp b/Engine/Plugins/Experimental/NeuralNetworkInference/Source/NeuralNetworkInferenceQA/Private/ModelUnitTester.cpp
--- a/Engine/Plugins/Experimental/NeuralNetworkInference/Source/NeuralNetworkInferenceQA/Private/ModelUnitTester.cpp
+++ b/Engine/Plugins/Experimental/NeuralNetworkInference/Source/NeuralNetworkInferenceQA/Private/ModelUnitTester.cpp
@@ -219,39 +219,48 @@ bool FModelUnitTester::ModelAccuracyTest(UNeuralNetwork* InOutNetwork, const ENe
 	InOutNetwork->SetBackEnd(InBackEnd);
 	const FString BackEndString = GetBackEndString(InBackEnd);
 	// Run each input with CPU/GPU and compare with each other and with the ground truth
+	// Multiple for loops to make sure that running on the CPU does not affect GPU results or vice versa
 	bool bDidGlobalTestPassed = true;
 	TArray<TArray<float>> CPUOutputs, CPUGPUCPUOutputs, CPUGPUGPUOutputs;
 	for (int32 Index = 0; Index < InputArrays.Num(); ++Index)
 	{
 		InOutNetwork->SetInputFromArrayCopy(InputArrays[Index]);
-		//// Input CPU + GPU + Output CPU
-		//InOutNetwork->SetDeviceType(ENeuralDeviceType::GPU);
-		//InOutNetwork->SetInputDeviceType(ENeuralDeviceType::CPU);
-		//InOutNetwork->SetOutputDeviceType(ENeuralDeviceType::CPU);
-		//InOutNetwork->Run();
-		//CPUGPUCPUOutputs.Emplace(InOutNetwork->GetOutputTensor().GetArrayCopy<float>());
 		// CPU
 		InOutNetwork->SetDeviceType(ENeuralDeviceType::CPU);
 		InOutNetwork->SetOutputDeviceType(ENeuralDeviceType::CPU);
 		InOutNetwork->SetInputDeviceType(ENeuralDeviceType::CPU);
 		InOutNetwork->Run();
 		CPUOutputs.Emplace(InOutNetwork->GetOutputTensor().GetArrayCopy<float>());
-// Input CPU + GPU + Output CPU
-InOutNetwork->SetDeviceType(ENeuralDeviceType::GPU);
-InOutNetwork->SetInputDeviceType(ENeuralDeviceType::CPU);
-InOutNetwork->SetOutputDeviceType(ENeuralDeviceType::CPU);
-InOutNetwork->Run();
-CPUGPUCPUOutputs.Emplace(InOutNetwork->GetOutputTensor().GetArrayCopy<float>());
-		//// Input CPU + GPU + Output GPU
-		//InOutNetwork->SetDeviceType(ENeuralDeviceType::GPU);
-		//InOutNetwork->SetInputDeviceType(ENeuralDeviceType::CPU);
-		//InOutNetwork->SetOutputDeviceType(ENeuralDeviceType::GPU);
-		//InOutNetwork->Run();
-		//InOutNetwork->OutputTensorsToCPU();
-CPUGPUGPUOutputs.Push(CPUGPUCPUOutputs[Index]); // @todo: Modify line to re-enable test
-		//CPUGPUGPUOutputs.Emplace(InOutNetwork->GetOutputTensor().GetArrayCopy<float>()); // @todo: Modify line to re-enable test
 	}
-	// The second for loop is to make sure running on the CPU does not affect GPU results or vice versa
+	for (int32 Index = 0; Index < InputArrays.Num(); ++Index)
+	{
+		InOutNetwork->SetInputFromArrayCopy(InputArrays[Index]);
+		// Input CPU + GPU + Output CPU
+		InOutNetwork->SetDeviceType(ENeuralDeviceType::GPU);
+		InOutNetwork->SetInputDeviceType(ENeuralDeviceType::CPU);
+		InOutNetwork->SetOutputDeviceType(ENeuralDeviceType::CPU);
+		InOutNetwork->Run();
+		CPUGPUCPUOutputs.Emplace(InOutNetwork->GetOutputTensor().GetArrayCopy<float>());
+	}
+	for (int32 Index = 0; Index < InputArrays.Num(); ++Index)
+	{
+// @todo: Remove if/else (and else contents) once DX12 version of NNI (UEAndORT) is ready
+if (InBackEnd == ENeuralBackEnd::UEOnly /*|| true*/)
+{
+		InOutNetwork->SetInputFromArrayCopy(InputArrays[Index]);
+		// Input CPU + GPU + Output GPU
+		InOutNetwork->SetDeviceType(ENeuralDeviceType::GPU);
+		InOutNetwork->SetInputDeviceType(ENeuralDeviceType::CPU);
+		InOutNetwork->SetOutputDeviceType(ENeuralDeviceType::GPU);
+		InOutNetwork->Run();
+		InOutNetwork->OutputTensorsToCPU();
+		CPUGPUGPUOutputs.Emplace(InOutNetwork->GetOutputTensor().GetArrayCopy<float>());
+}
+else
+{
+	CPUGPUGPUOutputs.Push(CPUGPUCPUOutputs[Index]);
+}
+	}
 	for (int32 Index = 0; Index < InputArrays.Num(); ++Index)
 	{
 		// Prepare verbose
