NeuralNetworkInference (NNI) plugin:
- Share D3D12 resources between UE and ONNX Runtime DirectML execution provider

#rnx

diff --git a/Engine/Plugins/Experimental/NeuralNetworkInference/Source/NeuralNetworkInference/Private/NeuralNetwork.cpp b/Engine/Plugins/Experimental/NeuralNetworkInference/Source/NeuralNetworkInference/Private/NeuralNetwork.cpp
--- a/Engine/Plugins/Experimental/NeuralNetworkInference/Source/NeuralNetworkInference/Private/NeuralNetwork.cpp
+++ b/Engine/Plugins/Experimental/NeuralNetworkInference/Source/NeuralNetworkInference/Private/NeuralNetwork.cpp
@@ -199,11 +199,6 @@ const FNeuralTensor& UNeuralNetwork::GetInputTensor(const int32 InTensorIndex) c
 	// UEAndORT
 	if (BackEndForCurrentPlatform == ENeuralBackEnd::UEAndORT)
 	{
-		if (ImplBackEndUEAndORT)
-		{
-			ImplBackEndUEAndORT->ClearResources();
-		}
-
 		return ImplBackEndUEAndORT->InputTensors[InTensorIndex];
 	}
 
@@ -454,7 +449,7 @@ bool UNeuralNetwork::Load()
 	// UEAndORT
 	if (BackEndForCurrentPlatform == ENeuralBackEnd::UEAndORT)
 	{
-		bIsLoaded = UNeuralNetwork::FImplBackEndUEAndORT::Load(ImplBackEndUEAndORT, AreInputTensorSizesVariable, ModelReadFromFileInBytes, ModelFullFilePath, GetDeviceType());
+		bIsLoaded = UNeuralNetwork::FImplBackEndUEAndORT::Load(ImplBackEndUEAndORT, AreInputTensorSizesVariable, ModelReadFromFileInBytes, ModelFullFilePath, GetDeviceType(), GetInputDeviceType(), GetOutputDeviceType());
 	}
 	// UEOnly
 	else if (BackEndForCurrentPlatform == ENeuralBackEnd::UEOnly)
diff --git a/Engine/Plugins/Experimental/NeuralNetworkInference/Source/NeuralNetworkInference/Private/NeuralNetworkImplBackEndUEAndORT.cpp b/Engine/Plugins/Experimental/NeuralNetworkInference/Source/NeuralNetworkInference/Private/NeuralNetworkImplBackEndUEAndORT.cpp
--- a/Engine/Plugins/Experimental/NeuralNetworkInference/Source/NeuralNetworkInference/Private/NeuralNetworkImplBackEndUEAndORT.cpp
+++ b/Engine/Plugins/Experimental/NeuralNetworkInference/Source/NeuralNetworkInference/Private/NeuralNetworkImplBackEndUEAndORT.cpp
@@ -3,6 +3,7 @@
 #include "NeuralNetworkImplBackEndUEAndORT.h"
 #include "NeuralNetworkInferenceUtils.h"
 #include "RedirectCoutAndCerrToUeLog.h"
+#include "NeuralNetworkInferenceUtilsGPU.h"
 
 #if defined(WITH_UE_AND_ORT_SUPPORT) && defined(PLATFORM_WIN64)
 	#include "HAL/CriticalSection.h"
@@ -185,7 +186,7 @@ IDMLDevice* FPrivateImplBackEndUEAndORT::FDMLDeviceList::Add(ID3D12Device* Devic
 //	return false;
 //}
 
-bool UNeuralNetwork::FImplBackEndUEAndORT::Load(TSharedPtr<FImplBackEndUEAndORT>& InOutImplBackEndUEAndORT, TArray<bool>& OutAreInputTensorSizesVariable, const TArray<uint8>& InModelReadFromFileInBytes, const FString& InModelFullFilePath, const ENeuralDeviceType InDeviceType)
+bool UNeuralNetwork::FImplBackEndUEAndORT::Load(TSharedPtr<FImplBackEndUEAndORT>& InOutImplBackEndUEAndORT, TArray<bool>& OutAreInputTensorSizesVariable, const TArray<uint8>& InModelReadFromFileInBytes, const FString& InModelFullFilePath, const ENeuralDeviceType InDeviceType, const ENeuralDeviceType InInputDeviceType, const ENeuralDeviceType InOutputDeviceType)
 {
 #ifdef WITH_UE_AND_ORT_SUPPORT
 #if WITH_EDITOR
@@ -194,6 +195,8 @@ bool UNeuralNetwork::FImplBackEndUEAndORT::Load(TSharedPtr<FImplBackEndUEAndORT>
 	{
 		const FRedirectCoutAndCerrToUeLog RedirectCoutAndCerrToUeLog;
 
+		InOutImplBackEndUEAndORT->ClearResources();
+
 		// Initialize and configure InOutImplBackEndUEAndORT
 		if (!UNeuralNetwork::FImplBackEndUEAndORT::InitializedAndConfigureMembers(InOutImplBackEndUEAndORT, InModelFullFilePath, InDeviceType))
 		{
@@ -223,9 +226,29 @@ bool UNeuralNetwork::FImplBackEndUEAndORT::Load(TSharedPtr<FImplBackEndUEAndORT>
 			UE_LOG(LogNeuralNetworkInference, Warning, TEXT("FImplBackEndUEAndORT::Load(): InModelReadFromFileInBytes was empty."));
 			return false;
 		}
+		
+		// Sanity check if device type is CPU and to make sure that input and/or output is also on the CPU
+		ENeuralDeviceType	InputDeviceType = InInputDeviceType;
+		ENeuralDeviceType	OutputDeviceType = InOutputDeviceType;
+
+		if (InDeviceType == ENeuralDeviceType::CPU && (InInputDeviceType == ENeuralDeviceType::GPU || InOutputDeviceType == ENeuralDeviceType::GPU))
+		{
+			UE_LOG(LogNeuralNetworkInference, Warning, TEXT("FImplBackEndUEAndORT::Load(): DeviceType is CPU but Input and/or Output is set to GPU, setting all to CPU."));
+			InputDeviceType = ENeuralDeviceType::CPU;
+			OutputDeviceType = ENeuralDeviceType::CPU;
+		}
+
+		if (!InOutImplBackEndUEAndORT->ConfigureTensors(InOutImplBackEndUEAndORT->InputTensors, &OutAreInputTensorSizesVariable, InputDeviceType, OutputDeviceType))
+		{
+			UE_LOG(LogNeuralNetworkInference, Warning, TEXT("FImplBackEndUEAndORT::Load(): Failed to configure input tensors."));
+			return false;
+		}
 
-		InOutImplBackEndUEAndORT->ConfigureTensors(InOutImplBackEndUEAndORT->InputTensors, &OutAreInputTensorSizesVariable);
-		InOutImplBackEndUEAndORT->ConfigureTensors(InOutImplBackEndUEAndORT->OutputTensors);
+		if (!InOutImplBackEndUEAndORT->ConfigureTensors(InOutImplBackEndUEAndORT->OutputTensors, nullptr, InputDeviceType, OutputDeviceType))
+		{
+			UE_LOG(LogNeuralNetworkInference, Warning, TEXT("FImplBackEndUEAndORT::Load(): Failed to configure output tensors."));
+			return false;
+		}
 
 		return true;
 	}
@@ -243,14 +266,24 @@ bool UNeuralNetwork::FImplBackEndUEAndORT::Load(TSharedPtr<FImplBackEndUEAndORT>
 #endif //WITH_UE_AND_ORT_SUPPORT
 }
 
+#ifdef WITH_UE_AND_ORT_SUPPORT
 void UNeuralNetwork::FImplBackEndUEAndORT::ClearResources()
 {
-#ifdef WITH_UE_AND_ORT_SUPPORT
 #ifdef PLATFORM_WIN64
-	DmlGPUAllocator.Reset(nullptr);
+	if (DmlGPUAllocator.IsValid() && DmlGPUAllocator->IsValid())
+	{ 
+		const int32 Num = DmlGPUResources.Num();
+		for (int32 Index = 0; Index < Num; ++Index)
+		{
+			DmlGPUAllocator->FreeGPUAllocation(DmlGPUResources[Index]);
+		}
+
+		DmlGPUResources.Reset(0);
+		DmlGPUAllocator.Reset(nullptr);
+	}
 #endif //PLATFORM_WIN64
-#endif //WITH_UE_AND_ORT_SUPPORT
 }
+#endif //WITH_UE_AND_ORT_SUPPORT
 
 void UNeuralNetwork::FImplBackEndUEAndORT::Run(const ENeuralNetworkSynchronousMode InSynchronousMode, const ENeuralDeviceType InInputDeviceType, const ENeuralDeviceType InOutputDeviceType)
 {
@@ -438,12 +471,13 @@ UE_LOG(LogNeuralNetworkInference, Display, TEXT("Testing the experimental DX12 (
 	return true;
 }
 
-void UNeuralNetwork::FImplBackEndUEAndORT::ConfigureTensors(TArray<FNeuralTensor>& OutTensors, TArray<bool>* OutAreInputTensorSizesVariable)
+bool UNeuralNetwork::FImplBackEndUEAndORT::ConfigureTensors(TArray<FNeuralTensor>& OutTensors, TArray<bool>* OutAreInputTensorSizesVariable, const ENeuralDeviceType InInputDeviceType, const ENeuralDeviceType InOutputDeviceType)
 {
 	const bool bIsInput = (OutAreInputTensorSizesVariable != nullptr);
 	TArray<const char*> TensorNames;
 	TArray<ENeuralDataType> TensorDataTypes;
 	TArray<TArray<int64>> TensorSizes;
+	TArray<ENeuralTensorTypeGPU> TensorGPUTypes;
 
 	const uint32 NumberTensors = bIsInput ? Session->GetInputCount() : Session->GetOutputCount();
 	if (OutAreInputTensorSizesVariable)
@@ -490,7 +524,7 @@ void UNeuralNetwork::FImplBackEndUEAndORT::ConfigureTensors(TArray<FNeuralTensor
 			{
 				TensorDataType = ENeuralDataType::None;
 				UE_LOG(LogNeuralNetworkInference, Warning, TEXT("FImplBackEndUEAndORT::ConfigureTensors(): ONNXTensorElementDataTypeEnum = %d not implemented yet."), (int32)ONNXTensorElementDataTypeEnum);
-				return;
+				return false;
 			}
 		}
 		TensorDataTypes.Push(TensorDataType);
@@ -520,21 +554,37 @@ void UNeuralNetwork::FImplBackEndUEAndORT::ConfigureTensors(TArray<FNeuralTensor
 		}
 		TensorSizes.Push(CurrentTensorSizes);
 
+		// @todo: Should caller specify tensor GPU type?
+		// Input tensor GPU type is set to Generic
+		// Output tensor GPU type is set to Output (i.e. data should not be copied from CPU)
+		ENeuralTensorTypeGPU	TensorGPUType;
+
+		if (bIsInput)
+		{
+			TensorGPUType = ENeuralTensorTypeGPU::Generic;
+		}
+		else
+		{
+			TensorGPUType = InOutputDeviceType == ENeuralDeviceType::GPU ? ENeuralTensorTypeGPU::Output : ENeuralTensorTypeGPU::Generic;
+		}
+
+		TensorGPUTypes.Push(TensorGPUType);
+
 		CurrentTypeInfo.release();
 	}
 
-	SetTensorsFromNetwork(OutTensors, TensorNames, TensorDataTypes, TensorSizes, bIsInput);
+	return SetTensorsFromNetwork(OutTensors, TensorNames, TensorDataTypes, TensorSizes, TensorGPUTypes, bIsInput);
 }
 
-void UNeuralNetwork::FImplBackEndUEAndORT::SetTensorsFromNetwork(TArray<FNeuralTensor>& OutTensors, TArray<const char*>& InTensorNames, TArray<ENeuralDataType>& InTensorDataTypes,
-	TArray<TArray<int64>>& InSizes, const bool bIsInput)
+bool UNeuralNetwork::FImplBackEndUEAndORT::SetTensorsFromNetwork(TArray<FNeuralTensor>& OutTensors, TArray<const char*>& InTensorNames, TArray<ENeuralDataType>& InTensorDataTypes,
+	TArray<TArray<int64>>& InSizes, TArray<ENeuralTensorTypeGPU>& InTensorGPUTypes, const bool bIsInput)
 {
 	const int32 TensorNumber = InTensorNames.Num();
 	if (InTensorDataTypes.Num() != TensorNumber || InSizes.Num() != TensorNumber)
 	{
 		UE_LOG(LogNeuralNetworkInference, Warning, TEXT("FImplBackEndUEAndORT::SetTensorsFromNetwork(): InTensorNames.Num() == InTensorDataTypes.Num() == InSizes.Num() failed, %d vs. %d vs. %d."),
 			InTensorNames.Num(), InTensorDataTypes.Num(), InSizes.Num());
-		return;
+		return false;
 	}
 
 	// Swap variables
@@ -563,9 +613,10 @@ void UNeuralNetwork::FImplBackEndUEAndORT::SetTensorsFromNetwork(TArray<FNeuralT
 	if (!bAreTensorsAlreadyCreatedWithRightNames)
 	{
 		OutTensors.Empty();
-		for (const char* const TensorName : TensorNames)
+		for (int32 TensorIndex = 0; TensorIndex < TensorNumber; ++TensorIndex)
 		{
-			OutTensors.Emplace(FNeuralTensor(ANSI_TO_TCHAR(TensorName)));
+			const char* TensorName = TensorNames[TensorIndex];
+			OutTensors.Emplace(FNeuralTensor(ANSI_TO_TCHAR(TensorName), InTensorGPUTypes[TensorIndex]));
 		}
 	}
 	ensureMsgf(OutTensors.Num() == TensorNumber, TEXT("OutTensors.Num() == TensorNumber failed, %d != %d."), OutTensors.Num(), TensorNumber);
@@ -578,11 +629,47 @@ void UNeuralNetwork::FImplBackEndUEAndORT::SetTensorsFromNetwork(TArray<FNeuralT
 		{
 			OrtTensors.Emplace(Ort::Value(nullptr));
 		}
+
+#ifdef PLATFORM_WIN64
+		if (InTensorGPUTypes[TensorIndex] == ENeuralTensorTypeGPU::Generic)
+		{
+			// Pre-allocate TArray (if size is different)
+			OutTensors[TensorIndex].SetNumUninitialized(InSizes[TensorIndex], InTensorDataTypes[TensorIndex]);
+			// Link tensor with ORT blob
+			LinkTensorToONNXRuntime(OutTensors, OrtTensors, *AllocatorInfo, TensorIndex);
+		}
+		else if (InTensorGPUTypes[TensorIndex] == ENeuralTensorTypeGPU::Output)
+		{
+			// @todo: should we remove this? It's currently used to read memory from GPU to CPU
+			OutTensors[TensorIndex].SetNumUninitialized(InSizes[TensorIndex], InTensorDataTypes[TensorIndex]);
+
+			OutTensors[TensorIndex].SetEnableGPU(true);
+
+			// @todo: This requires SetNumUnitialized() to be run, otherwise Size and Volume will be set to 0
+			void* D3DResource = nullptr;
+
+			if (!OutTensors[TensorIndex].InitPooledBuffer(&D3DResource))
+			{
+				UE_LOG(LogNeuralNetworkInference, Warning, TEXT("FImplBackEndUEAndORT::SetTensorsFromNetwork(): Failed to initialize pooled buffer"));
+				return false;
+			}
+
+			// Link tensor with ORT blob
+			if (!LinkTensorResourceToONNXRuntime(OutTensors[TensorIndex], OrtTensors[TensorIndex], D3DResource))
+			{
+				UE_LOG(LogNeuralNetworkInference, Warning, TEXT("FImplBackEndUEAndORT::SetTensorsFromNetwork(): Failed to link GPU resource to ONNX runtime"));
+				return false;
+			}
+		}
+#else
 		// Pre-allocate TArray (if size is different)
 		OutTensors[TensorIndex].SetNumUninitialized(InSizes[TensorIndex], InTensorDataTypes[TensorIndex]);
 		// Link tensor with ORT blob
 		LinkTensorToONNXRuntime(OutTensors, OrtTensors, *AllocatorInfo, TensorIndex);
+#endif
 	}
+
+	return true;
 }
 
 void UNeuralNetwork::FImplBackEndUEAndORT::LinkTensorToONNXRuntime(TArray<FNeuralTensor>& InOutTensors, TArray<Ort::Value>& InOutOrtTensors, Ort::MemoryInfo& InOutAllocatorInfo, const int32 InTensorIndex)
@@ -618,4 +705,57 @@ void UNeuralNetwork::FImplBackEndUEAndORT::LinkTensorToONNXRuntime(TArray<FNeura
 	}
 }
 
+#ifdef PLATFORM_WIN64
+
+bool UNeuralNetwork::FImplBackEndUEAndORT::LinkTensorResourceToONNXRuntime(FNeuralTensor& InOutTensor, Ort::Value& InOutOrtTensor, void* D3DResource)
+{
+	if (DmlGPUAllocator.IsValid() && !DmlGPUAllocator->IsValid())
+	{
+		UE_LOG(LogNeuralNetworkInference, Warning, TEXT("FImplBackEndUEAndORT::LinkTensorResourceToONNXRuntime(): DmlGPUAllocator is not valid"));
+		return false;
+	}
+
+	void* DmlGPUAllocation = DmlGPUAllocator->GPUAllocationFromD3DResource(D3DResource);
+	if (!DmlGPUAllocation)
+	{
+		UE_LOG(LogNeuralNetworkInference, Warning, TEXT("FImplBackEndUEAndORT::LinkTensorResourceToONNXRuntime(): DmlGPUAllocation is NULL"));
+		return false;
+	}
+
+	DmlGPUResources.Emplace(DmlGPUAllocation);
+
+	const TArray<int64>& Sizes = InOutTensor.GetSizes();
+	if (Sizes.Num() > 0 && InOutTensor.Num() > 0)
+	{
+		const int32 ArrayDimensions = Sizes.Num();
+
+		const ENeuralDataType NeuralDataType = InOutTensor.GetDataType();
+		if (NeuralDataType == ENeuralDataType::Float)
+		{
+#ifdef _WIN32
+			const TArray<int64_t>& SizesInt64t = Sizes;
+#else
+			checkf(sizeof(int64) == sizeof(int64_t), TEXT("int64 and int64_t should both have the same size."));
+			TArray<int64_t> SizesInt64t;
+			SizesInt64t.SetNumUninitialized(ArrayDimensions);
+			FMemory::Memcpy(SizesInt64t.GetData(), (int64_t*)Sizes.GetData(), sizeof(int64_t) * ArrayDimensions);
+#endif //_WIN32
+			InOutOrtTensor = Ort::Value::CreateTensor(DmlGPUAllocator->GetProviderMemoryInfo(), DmlGPUAllocation, InOutTensor.NumInBytes(), SizesInt64t.GetData(), ArrayDimensions, ONNX_TENSOR_ELEMENT_DATA_TYPE_FLOAT);
+		}
+		//else if (NeuralDataType == ENeuralDataType::Double)
+		//{
+		//	InOutOrtTensor = Ort::Value::CreateTensor(DmlGPUAllocator->GetProviderMemoryInfo(), DmlGPUAllocation, InOutTensor.NumInBytes(), SizesInt64t.GetData(), ArrayDimensions, ONNX_TENSOR_ELEMENT_DATA_TYPE_DOUBLE);
+		//}
+		else
+		{
+			UE_LOG(LogNeuralNetworkInference, Warning, TEXT("FImplBackEndUEAndORT::LinkTensorToONNXRuntime(): Not implemented (yet) for ENeuralDataType = %d."), (int32)NeuralDataType);
+			return false;
+		}
+	}
+
+	return true;
+}
+
+#endif
+
 #endif //WITH_UE_AND_ORT_SUPPORT
diff --git a/Engine/Plugins/Experimental/NeuralNetworkInference/Source/NeuralNetworkInference/Private/NeuralNetworkImplBackEndUEAndORT.h b/Engine/Plugins/Experimental/NeuralNetworkInference/Source/NeuralNetworkInference/Private/NeuralNetworkImplBackEndUEAndORT.h
--- a/Engine/Plugins/Experimental/NeuralNetworkInference/Source/NeuralNetworkInference/Private/NeuralNetworkImplBackEndUEAndORT.h
+++ b/Engine/Plugins/Experimental/NeuralNetworkInference/Source/NeuralNetworkInference/Private/NeuralNetworkImplBackEndUEAndORT.h
@@ -41,6 +41,7 @@ public:
 #ifdef PLATFORM_WIN64
 	TUniquePtr<OrtDMLProviderOptions> DmlProviderOptions; /* DirectML execution provider options */
 	TUniquePtr<Ort::DMLGPUResourceAllocator> DmlGPUAllocator; /* DirectML GPU allocator */
+	TArray<void*> DmlGPUResources; /* Shared D3D12 resources with DirectML GPU execution provider */
 #endif
 	TArray<Ort::Value> InputOrtTensors; /* Actual ONNXRuntime tensors */
 	TArray<const char*> InputTensorNames; /* Tensor names */
@@ -48,10 +49,8 @@ public:
 	TArray<const char*> OutputTensorNames; /* Tensor names */
 #endif //WITH_UE_AND_ORT_SUPPORT
 
-	static bool Load(TSharedPtr<FImplBackEndUEAndORT>& InOutImplBackEndUEAndORT, TArray<bool>& OutAreInputTensorSizesVariable, const TArray<uint8>& InModelReadFromFileInBytes, const FString& InModelFullFilePath, const ENeuralDeviceType InDeviceType);
+	static bool Load(TSharedPtr<FImplBackEndUEAndORT>& InOutImplBackEndUEAndORT, TArray<bool>& OutAreInputTensorSizesVariable, const TArray<uint8>& InModelReadFromFileInBytes, const FString& InModelFullFilePath, const ENeuralDeviceType InDeviceType, const ENeuralDeviceType InInputDeviceType, const ENeuralDeviceType InOutputDeviceType);
 	
-	void ClearResources();
-
 	void Run(const ENeuralNetworkSynchronousMode InSynchronousMode, const ENeuralDeviceType InInputDeviceType, const ENeuralDeviceType InOutputDeviceType);
 
 #ifdef WITH_UE_AND_ORT_SUPPORT
@@ -60,10 +59,17 @@ private:
 
 	bool ConfigureMembers(const ENeuralDeviceType InDeviceType);
 
-	void ConfigureTensors(TArray<FNeuralTensor>& OutTensors, TArray<bool>* OutAreInputTensorSizesVariable = nullptr);
+	bool ConfigureTensors(TArray<FNeuralTensor>& OutTensors, TArray<bool>* OutAreInputTensorSizesVariable, const ENeuralDeviceType InInputDeviceType, const ENeuralDeviceType InOutputDeviceType);
 
-	void SetTensorsFromNetwork(TArray<FNeuralTensor>& OutTensors, TArray<const char*>& InTensorNames, TArray<ENeuralDataType>& InTensorDataTypes, TArray<TArray<int64>>& InSizes, const bool bIsInput);
+	bool SetTensorsFromNetwork(TArray<FNeuralTensor>& OutTensors, TArray<const char*>& InTensorNames, TArray<ENeuralDataType>& InTensorDataTypes, TArray<TArray<int64>>& InSizes, TArray<ENeuralTensorTypeGPU>& InTensorGPUTypes, const bool bIsInput);
 
 	static void LinkTensorToONNXRuntime(TArray<FNeuralTensor>& InOutTensors, TArray<Ort::Value>& InOutOrtTensors, Ort::MemoryInfo& InOutAllocatorInfo, const int32 InTensorIndex);
+
+#ifdef PLATFORM_WIN64
+	bool LinkTensorResourceToONNXRuntime(FNeuralTensor& InOutTensor, Ort::Value& InOutOrtTensor, void* D3DResource);
+#endif
+
+	void ClearResources();
+
 #endif //WITH_UE_AND_ORT_SUPPORT
 };
diff --git a/Engine/Plugins/Experimental/NeuralNetworkInference/Source/NeuralNetworkInference/Private/NeuralNetworkInferenceUtilsGPU.cpp b/Engine/Plugins/Experimental/NeuralNetworkInference/Source/NeuralNetworkInference/Private/NeuralNetworkInferenceUtilsGPU.cpp
--- a/Engine/Plugins/Experimental/NeuralNetworkInference/Source/NeuralNetworkInference/Private/NeuralNetworkInferenceUtilsGPU.cpp
+++ b/Engine/Plugins/Experimental/NeuralNetworkInference/Source/NeuralNetworkInference/Private/NeuralNetworkInferenceUtilsGPU.cpp
@@ -4,7 +4,29 @@
 #include "NeuralNetworkInferenceUtils.h"
 #include "NeuralTensorResourceArray.h"
 
+#ifdef PLATFORM_WIN64
 
+#include "RHIResources.h"
+#include "ShaderCore.h"
+
+// Disable NOMINMAX & WIN32_LEAN_AND_MEAN defines to avoid compiler warnings
+#pragma push_macro("NOMINMAX")
+#pragma push_macro("WIN32_LEAN_AND_MEAN")
+#undef NOMINMAX
+#undef WIN32_LEAN_AND_MEAN
+
+#include "D3D12RHI.h"
+#include "D3D12RHICommon.h"
+
+// @todo: We need D3D12RHIPrivate.h in order to compile correctly, i.e. to be able to use FD3D12Resource and ID3D12Resource
+#include "D3D12RHIPrivate.h"
+
+#include "D3D12Resources.h"
+
+#pragma pop_macro("WIN32_LEAN_AND_MEAN")
+#pragma pop_macro("NOMINMAX")
+
+#endif
 
 /* FNeuralNetworkInferenceUtilsGPU public functions
  *****************************************************************************/
@@ -46,3 +68,27 @@ bool FNeuralNetworkInferenceUtilsGPU::GPUSanityChecks(const FRDGBuilder* const I
 	}
 	return GPUSanityChecks(InGraphBuilder);
 }
+
+
+#ifdef PLATFORM_WIN64
+
+bool FNeuralNetworkInferenceUtilsGPU::IsD3D12RHI()
+{
+	const FString RHIName = GDynamicRHI->GetName();
+
+	return RHIName == TEXT("D3D12");
+}
+
+ID3D12Resource* FNeuralNetworkInferenceUtilsGPU::GetD3D12Resource(FRHIBuffer* Buffer)
+{
+	if (!IsD3D12RHI())
+	{
+		return nullptr;
+	}
+
+	FD3D12Buffer*	D3DBuffer = static_cast<FD3D12Buffer*>(Buffer);
+
+	return D3DBuffer->GetResource()->GetResource();
+}
+
+#endif
\ No newline at end of file
diff --git a/Engine/Plugins/Experimental/NeuralNetworkInference/Source/NeuralNetworkInference/Private/NeuralNetworkInferenceUtilsGPU.h b/Engine/Plugins/Experimental/NeuralNetworkInference/Source/NeuralNetworkInference/Private/NeuralNetworkInferenceUtilsGPU.h
--- a/Engine/Plugins/Experimental/NeuralNetworkInference/Source/NeuralNetworkInference/Private/NeuralNetworkInferenceUtilsGPU.h
+++ b/Engine/Plugins/Experimental/NeuralNetworkInference/Source/NeuralNetworkInference/Private/NeuralNetworkInferenceUtilsGPU.h
@@ -6,6 +6,10 @@
 #include "RenderGraphBuilder.h"
 #include "RHI.h"
 
+#ifdef PLATFORM_WIN64
+struct ID3D12Resource;
+#endif
+
 class FNeuralNetworkInferenceUtilsGPU
 {
 public:
@@ -20,4 +24,11 @@ public:
 	 */
 	static bool GPUSanityChecks(const FRDGBuilder* const InGraphBuilder);
 	static bool GPUSanityChecks(const FRDGBuilder* const InGraphBuilder, const bool bInIsLoaded);
+
+
+#ifdef PLATFORM_WIN64
+	static bool IsD3D12RHI();
+
+	static ID3D12Resource* GetD3D12Resource(FRHIBuffer* Buffer);
+#endif
 };
diff --git a/Engine/Plugins/Experimental/NeuralNetworkInference/Source/NeuralNetworkInference/Private/NeuralTensor.cpp b/Engine/Plugins/Experimental/NeuralNetworkInference/Source/NeuralNetworkInference/Private/NeuralTensor.cpp
--- a/Engine/Plugins/Experimental/NeuralNetworkInference/Source/NeuralNetworkInference/Private/NeuralTensor.cpp
+++ b/Engine/Plugins/Experimental/NeuralNetworkInference/Source/NeuralNetworkInference/Private/NeuralTensor.cpp
@@ -370,6 +370,61 @@ void FNeuralTensor::ToCPU_RenderThread(FRDGBuilder* InOutGraphBuilder)
 	});
 }
 
+bool FNeuralTensor::InitPooledBuffer(void** NativeResource)
+{
+	if (!bEnableGPU)
+	{
+		UE_LOG(LogNeuralNetworkInference, Warning, TEXT("FNeuralTensor::InitPooledBuffer():bEnableGPU is false"));
+		return false;
+	}
+
+	// Make sure that we wait for Graph builder to finish
+	std::atomic<bool> bIsGraphBuilderDone(false);
+
+	ENQUEUE_RENDER_COMMAND(FNeuralTensor_InitPooledBuffer)(
+		[this, &NativeResource, &bIsGraphBuilderDone] (FRHICommandListImmediate& RHICmdList)
+		{
+			FRDGBuilder Builder(RHICmdList);
+
+			FRDGBufferDesc BufferDesc;
+			BufferDesc.BytesPerElement = FDataType::GetSize(DataType);
+			BufferDesc.NumElements = Num();
+			BufferDesc.Usage = BUF_UnorderedAccess | BUF_ShaderResource;
+			BufferDesc.UnderlyingType = FRDGBufferDesc::EUnderlyingType::VertexBuffer;
+
+			FRDGBufferRef BufferRef = Builder.CreateBuffer(BufferDesc, *GetName());
+			
+			// Recreate PooledBuffer for future runs
+			PooledBuffer = MakeShared<TRefCountPtr<FRDGPooledBuffer>>();
+			*PooledBuffer = Builder.ConvertToExternalBuffer(BufferRef);
+
+			// @todo: For now we're not initializing SRV and UAV, since those are not used by DirectML execution provider
+			BufferSRVRef.Reset();
+			BufferUAVRef.Reset();
+
+			Builder.Execute();
+
+#ifdef PLATFORM_WIN64
+			if (NativeResource)
+			{
+				FRHIBuffer* Buffer = (*PooledBuffer)->GetRHI();
+
+				NativeResource[0] = FNeuralNetworkInferenceUtilsGPU::GetD3D12Resource(Buffer);
+			}
+#endif
+
+			bIsGraphBuilderDone = true;
+		}
+	);
+
+	while (!bIsGraphBuilderDone)
+	{
+		FPlatformProcess::Sleep(0.1e-3);
+	}
+
+	return true;
+}
+
 TRefCountPtr<class FRDGPooledBuffer>& FNeuralTensor::GetPooledBuffer() const
 {
 	// Sanity checks
diff --git a/Engine/Plugins/Experimental/NeuralNetworkInference/Source/NeuralNetworkInference/Public/NeuralTensor.h b/Engine/Plugins/Experimental/NeuralNetworkInference/Source/NeuralNetworkInference/Public/NeuralTensor.h
--- a/Engine/Plugins/Experimental/NeuralNetworkInference/Source/NeuralNetworkInference/Public/NeuralTensor.h
+++ b/Engine/Plugins/Experimental/NeuralNetworkInference/Source/NeuralNetworkInference/Public/NeuralTensor.h
@@ -196,6 +196,12 @@ public:
 	void UpdateSRVAndOrUAV_RenderThread(FRDGBuilder* InOutGraphBuilder);
 	void ToCPU_RenderThread(FRDGBuilder* InOutGraphBuilder);
 
+	/**
+	 * Allocate data for GPU pooled buffer.
+	 * NativeResource is a pointer to ID3D12Resource obtained from PooledBuffer and that one can be shared by DirectML execution provider
+	 */
+	bool InitPooledBuffer(void** NativeResource = nullptr);
+
 	TRefCountPtr<class FRDGPooledBuffer>& GetPooledBuffer() const;
 	const FRDGBufferSRVRef GetBufferSRVRef() const;
 	FRDGBufferUAVRef GetBufferUAVRef();
