NeuralNetworkInference (NNI) plugin:
- Popup + UE_LOG(Warning) if DX12 not selected.
#rnx

diff --git a/Engine/Plugins/Experimental/NeuralNetworkInference/Source/NeuralNetworkInference/Private/NeuralNetwork.cpp b/Engine/Plugins/Experimental/NeuralNetworkInference/Source/NeuralNetworkInference/Private/NeuralNetwork.cpp
--- a/Engine/Plugins/Experimental/NeuralNetworkInference/Source/NeuralNetworkInference/Private/NeuralNetwork.cpp
+++ b/Engine/Plugins/Experimental/NeuralNetworkInference/Source/NeuralNetworkInference/Private/NeuralNetwork.cpp
@@ -449,6 +449,7 @@ bool UNeuralNetwork::Load()
 	// UEAndORT
 	if (BackEndForCurrentPlatform == ENeuralBackEnd::UEAndORT)
 	{
+		UNeuralNetwork::FImplBackEndUEAndORT::SetDeviceToCPUIfD3D12NotSupported(DeviceType);
 		bIsLoaded = UNeuralNetwork::FImplBackEndUEAndORT::Load(ImplBackEndUEAndORT, AreInputTensorSizesVariable, ModelReadFromFileInBytes, ModelFullFilePath, GetDeviceType(), GetInputDeviceType(), GetOutputDeviceType());
 	}
 	// UEOnly
diff --git a/Engine/Plugins/Experimental/NeuralNetworkInference/Source/NeuralNetworkInference/Private/NeuralNetworkImplBackEndUEAndORT.cpp b/Engine/Plugins/Experimental/NeuralNetworkInference/Source/NeuralNetworkInference/Private/NeuralNetworkImplBackEndUEAndORT.cpp
--- a/Engine/Plugins/Experimental/NeuralNetworkInference/Source/NeuralNetworkInference/Private/NeuralNetworkImplBackEndUEAndORT.cpp
+++ b/Engine/Plugins/Experimental/NeuralNetworkInference/Source/NeuralNetworkInference/Private/NeuralNetworkImplBackEndUEAndORT.cpp
@@ -2,8 +2,11 @@
 
 #include "NeuralNetworkImplBackEndUEAndORT.h"
 #include "NeuralNetworkInferenceUtils.h"
-#include "RedirectCoutAndCerrToUeLog.h"
 #include "NeuralNetworkInferenceUtilsGPU.h"
+#if WITH_EDITOR
+#include "Misc/MessageDialog.h"
+#endif //WITH_EDITOR
+#include "RedirectCoutAndCerrToUeLog.h"
 
 #if defined(WITH_UE_AND_ORT_SUPPORT) && defined(PLATFORM_WIN64)
 	#include "HAL/CriticalSection.h"
@@ -120,6 +123,55 @@ IDMLDevice* FPrivateImplBackEndUEAndORT::FDMLDeviceList::Add(ID3D12Device* Devic
 /* UNeuralNetwork public functions
  *****************************************************************************/
 
+void UNeuralNetwork::FImplBackEndUEAndORT::SetDeviceToCPUIfD3D12NotSupported(ENeuralDeviceType& InOutDeviceType)
+{
+	if (!IsD3D12Supported())
+	{
+		UE_LOG(LogNeuralNetworkInference, Warning, TEXT("FImplBackEndUEAndORT::SetDeviceToCPUIfD3D12NotSupported(): Setting device to CPU provisionally."));
+		InOutDeviceType = ENeuralDeviceType::CPU;
+	}
+}
+
+bool UNeuralNetwork::FImplBackEndUEAndORT::IsD3D12Supported()
+{
+#ifdef WITH_UE_AND_ORT_SUPPORT
+
+#ifdef PLATFORM_WIN64
+	// To create a DirectML device we need to check that we're using DX12 first
+	const FString RHIName = GDynamicRHI->GetName();
+
+	// Display warning if not DX12
+	if (RHIName != TEXT("D3D12"))
+	{
+		const FString ErrorMessage = TEXT("On Windows, only DirectX 12 rendering (\"D3D12\") is compatible with NeuralNetworkInference (NNI). Instead, \"") + RHIName
+			+ TEXT("\" was used. You have the following options:\n"
+				"\t1. (Recommended) Switch Unreal Engine to DX12. In order to do that:\n"
+					"\t\t - Go to \"Project Settings\", \"Platforms\", \"Windows\", \"Default RHI\".\n"
+					"\t\t - Select \"DirectX 12\".\n"
+					"\t\t - Restart Unreal Engine.\n"
+				"\t2. Alternatively, switch the network to CPU with UNeuralNetwork::SetDeviceType().\n"
+				"\t3. (Not recommended) You could also switch the network to UEOnly with UNeuralNetwork::SetBackEnd().\n"
+				"Setting the network to CPU provisionally.");
+		UE_LOG(LogNeuralNetworkInference, Warning, TEXT("FImplBackEndUEAndORT::ConfigureMembers(): %s"), *ErrorMessage, *RHIName);
+#if WITH_EDITOR
+		FMessageDialog::Open(EAppMsgType::Ok, FText::FromString(ErrorMessage));
+#endif //WITH_EDITOR
+		return false;
+	}
+
+	return true;
+
+#else //PLATFORM_WIN64
+	UE_LOG(LogNeuralNetworkInference, Warning, TEXT("FImplBackEndUEAndORT::IsD3D12Supported(): This function will always return false outside of Windows."));
+	return false;
+#endif //PLATFORM_WIN64
+
+#else //WITH_UE_AND_ORT_SUPPORT
+	UE_LOG(LogNeuralNetworkInference, Warning, TEXT("FImplBackEndUEAndORT::IsD3D12Supported(): This function will always return false if WITH_UE_AND_ORT_SUPPORT is not defined."));
+	return false;
+#endif //WITH_UE_AND_ORT_SUPPORT
+}
+
 //bool UNeuralNetwork::FImplBackEndUEAndORT::Load(TSharedPtr<FImplBackEndUEAndORT>& InOutImplBackEndUEAndORT, TArray<bool>& OutAreInputTensorSizesVariable, const TArray<uint8>& InModelReadFromFileInBytes, const FString& InModelFullFilePath, const ENeuralDeviceType InDeviceType)
 //{
 //	// Initialize and configure InOutImplBackEndUEAndORT
@@ -378,24 +430,10 @@ bool UNeuralNetwork::FImplBackEndUEAndORT::ConfigureMembers(const ENeuralDeviceT
 	if (InDeviceType == ENeuralDeviceType::GPU)
 	{
 #ifdef PLATFORM_WIN64
-		// To create a DirectML device we need to check that we're using DX12 first
-		const FString RHIName = GDynamicRHI->GetName();
-
-		// @todo for Nebojša.Dragosavac & Ginés.Hidalgo: Handle the default DX11 nicely
-		if (RHIName != TEXT("D3D12"))
+		if (!IsD3D12Supported())
 		{
-			//UE_LOG(LogNeuralNetworkInference, Warning, TEXT("FImplBackEndUEAndORT::ConfigureMembers(): Only DX12 (\"D3D12\") rendering is compatible with UEAndORT-based NNI. Instead, \"%s\" was used."), *RHIName);
-			//return false;
-
-			// Temporary workaround to avoid breaking NNI
-			if (OrtSessionOptionsAppendExecutionProvider_DML(*SessionOptions, 0))
-			{
-				UE_LOG(LogNeuralNetworkInference, Warning, TEXT("FImplBackEndUEAndORT::ConfigureMembers(): Some error occurred when using OrtSessionOptionsAppendExecutionProvider_DML()."));
-				return false;
-			}
-			return true;
+			return false;
 		}
-UE_LOG(LogNeuralNetworkInference, Display, TEXT("Testing the experimental DX12 (\"D3D12\") back end."));
 
 		// Get adapter's D3D12 device that we would like to share with DirectML execution provider
 		// NOTE: For now we're only using first device that has Dadapter 0 and device 0
@@ -434,7 +472,7 @@ UE_LOG(LogNeuralNetworkInference, Display, TEXT("Testing the experimental DX12 (
 		DmlProviderOptions->cmd_queue = NativeCmdQ;
 		DmlProviderOptions->resource_allocator = DmlGPUAllocator->GetAllocatorAddressOf();
 
-		if (OrtSessionOptionsAppendExecutionProviderWithOptions_DML(*SessionOptions, DmlProviderOptions.Get()))
+		if (OrtSessionOptionsAppendExecutionProviderWithOptions_DML(*SessionOptions, DmlProviderOptions.Get())) // OrtSessionOptionsAppendExecutionProvider_DML(*SessionOptions, 0) without sharing
 		{
 			UE_LOG(LogNeuralNetworkInference, Warning, TEXT("FImplBackEndUEAndORT::ConfigureMembers(): Some error occurred when using OrtSessionOptionsAppendExecutionProviderEx_DML()."));
 			return false;
@@ -709,7 +747,7 @@ void UNeuralNetwork::FImplBackEndUEAndORT::LinkTensorToONNXRuntime(TArray<FNeura
 
 bool UNeuralNetwork::FImplBackEndUEAndORT::LinkTensorResourceToONNXRuntime(FNeuralTensor& InOutTensor, Ort::Value& InOutOrtTensor, void* D3DResource)
 {
-	if (DmlGPUAllocator.IsValid() && !DmlGPUAllocator->IsValid())
+	if (!DmlGPUAllocator.IsValid() || !DmlGPUAllocator->IsValid())
 	{
 		UE_LOG(LogNeuralNetworkInference, Warning, TEXT("FImplBackEndUEAndORT::LinkTensorResourceToONNXRuntime(): DmlGPUAllocator is not valid"));
 		return false;
diff --git a/Engine/Plugins/Experimental/NeuralNetworkInference/Source/NeuralNetworkInference/Private/NeuralNetworkImplBackEndUEAndORT.h b/Engine/Plugins/Experimental/NeuralNetworkInference/Source/NeuralNetworkInference/Private/NeuralNetworkImplBackEndUEAndORT.h
--- a/Engine/Plugins/Experimental/NeuralNetworkInference/Source/NeuralNetworkInference/Private/NeuralNetworkImplBackEndUEAndORT.h
+++ b/Engine/Plugins/Experimental/NeuralNetworkInference/Source/NeuralNetworkInference/Private/NeuralNetworkImplBackEndUEAndORT.h
@@ -49,6 +49,10 @@ public:
 	TArray<const char*> OutputTensorNames; /* Tensor names */
 #endif //WITH_UE_AND_ORT_SUPPORT
 
+	static void SetDeviceToCPUIfD3D12NotSupported(ENeuralDeviceType& InOutDeviceType);
+
+	static bool IsD3D12Supported();
+
 	static bool Load(TSharedPtr<FImplBackEndUEAndORT>& InOutImplBackEndUEAndORT, TArray<bool>& OutAreInputTensorSizesVariable, const TArray<uint8>& InModelReadFromFileInBytes, const FString& InModelFullFilePath, const ENeuralDeviceType InDeviceType, const ENeuralDeviceType InInputDeviceType, const ENeuralDeviceType InOutputDeviceType);
 	
 	void Run(const ENeuralNetworkSynchronousMode InSynchronousMode, const ENeuralDeviceType InInputDeviceType, const ENeuralDeviceType InOutputDeviceType);
