NeuralNetworkInference plugin:
Disconnect Mac until warnings are fixed. Meeting with Will to study them

diff --git a/Engine/Plugins/Experimental/NeuralNetworkInference/NeuralNetworkInference.uplugin b/Engine/Plugins/Experimental/NeuralNetworkInference/NeuralNetworkInference.uplugin
--- a/Engine/Plugins/Experimental/NeuralNetworkInference/NeuralNetworkInference.uplugin
+++ b/Engine/Plugins/Experimental/NeuralNetworkInference/NeuralNetworkInference.uplugin
@@ -40,7 +40,7 @@
 			"Type": "RuntimeAndProgram",
 			"LoadingPhase": "Default",
 			"WhitelistPlatforms": [
-				"Win64", "Linux", "Mac"
+				"Win64", "Linux"
 			]
 		},
 		{
