NeuralNetworkInference (NNI) plugin
- Modified InitPooledBuffer() to create BufferUAV and BufferSRV since those can be used by outside code (MLDeformer for example)

#rnx

diff --git a/Engine/Plugins/Experimental/NeuralNetworkInference/Source/NeuralNetworkInference/Private/NeuralTensor.cpp b/Engine/Plugins/Experimental/NeuralNetworkInference/Source/NeuralNetworkInference/Private/NeuralTensor.cpp
--- a/Engine/Plugins/Experimental/NeuralNetworkInference/Source/NeuralNetworkInference/Private/NeuralTensor.cpp
+++ b/Engine/Plugins/Experimental/NeuralNetworkInference/Source/NeuralNetworkInference/Private/NeuralTensor.cpp
@@ -398,9 +398,9 @@ bool FNeuralTensor::InitPooledBuffer(void** NativeResource)
 			PooledBuffer = MakeShared<TRefCountPtr<FRDGPooledBuffer>>();
 			*PooledBuffer = Builder.ConvertToExternalBuffer(BufferRef);
 
-			// @todo: For now we're not initializing SRV and UAV, since those are not used by DirectML execution provider
-			BufferSRVRef.Reset();
-			BufferUAVRef.Reset();
+			// SRV and UAV can be used by callers (for example MLDeformer), look at the GetPooledBuffer()
+			BufferSRVRef = MakeShared<FRDGBufferSRVRef>(Builder.CreateSRV(BufferRef, FDataType::GetPixelFormat(DataType)));
+			BufferUAVRef = MakeShared<FRDGBufferUAVRef>(Builder.CreateUAV(BufferRef, FDataType::GetPixelFormat(DataType)));
 
 			Builder.Execute();
 
