NeuralNetworkInference (NNI) plugin:
- NNIUEAndORT do not pop up annoying warning if CPU already.

_codereview Francisco.Vicente
#rnx

diff --git a/Engine/Plugins/Experimental/NeuralNetworkInference/Source/NeuralNetworkInference/Private/NeuralNetworkImplBackEndUEAndORT.cpp b/Engine/Plugins/Experimental/NeuralNetworkInference/Source/NeuralNetworkInference/Private/NeuralNetworkImplBackEndUEAndORT.cpp
--- a/Engine/Plugins/Experimental/NeuralNetworkInference/Source/NeuralNetworkInference/Private/NeuralNetworkImplBackEndUEAndORT.cpp
+++ b/Engine/Plugins/Experimental/NeuralNetworkInference/Source/NeuralNetworkInference/Private/NeuralNetworkImplBackEndUEAndORT.cpp
@@ -125,10 +125,13 @@ IDMLDevice* FPrivateImplBackEndUEAndORT::FDMLDeviceList::Add(ID3D12Device* Devic
 
 void UNeuralNetwork::FImplBackEndUEAndORT::SetDeviceToCPUIfD3D12NotSupported(ENeuralDeviceType& InOutDeviceType)
 {
-	if (!IsD3D12Supported())
+	if (InOutDeviceType != ENeuralDeviceType::CPU)
 	{
-		UE_LOG(LogNeuralNetworkInference, Warning, TEXT("FImplBackEndUEAndORT::SetDeviceToCPUIfD3D12NotSupported(): Setting device to CPU provisionally."));
-		InOutDeviceType = ENeuralDeviceType::CPU;
+		if (!IsD3D12Supported())
+		{
+			UE_LOG(LogNeuralNetworkInference, Warning, TEXT("FImplBackEndUEAndORT::SetDeviceToCPUIfD3D12NotSupported(): Setting device to CPU provisionally."));
+			InOutDeviceType = ENeuralDeviceType::CPU;
+		}
 	}
 }
 
