NeuralNetworkInference (NNI) plugin:
- Enabled GPU output test in ModelUnitTester.cpp.
- Fixed bug in SetDeviceType: SetDeviceType/Input/Output reduced into a single SetDeviceType, to avoid bugs when setting different ones and having the re-load the network.

_codereview nebojsa.dragosavac
#rnx

diff --git a/Engine/Plugins/Experimental/MLDeformer/Source/MLDeformer/Private/MLDeformerAsset.cpp b/Engine/Plugins/Experimental/MLDeformer/Source/MLDeformer/Private/MLDeformerAsset.cpp
--- a/Engine/Plugins/Experimental/MLDeformer/Source/MLDeformer/Private/MLDeformerAsset.cpp
+++ b/Engine/Plugins/Experimental/MLDeformer/Source/MLDeformer/Private/MLDeformerAsset.cpp
@@ -65,9 +65,7 @@ void UMLDeformerAsset::PostLoad()
 	if (NeuralNetwork)
 	{
 		NeuralNetwork->SetBackEnd(ENeuralBackEnd::Auto);
-		NeuralNetwork->SetInputDeviceType(ENeuralDeviceType::CPU);
-		NeuralNetwork->SetDeviceType(ENeuralDeviceType::GPU);
-		NeuralNetwork->SetOutputDeviceType(ENeuralDeviceType::GPU);
+		NeuralNetwork->SetDeviceType(/*DeviceType*/ENeuralDeviceType::GPU, /*InputDeviceType*/ENeuralDeviceType::CPU, /*OutputDeviceType*/ENeuralDeviceType::GPU);
 	}
 
 #if WITH_EDITOR
diff --git a/Engine/Plugins/Experimental/MLDeformer/Source/MLDeformerEditor/Private/MLDeformerEditorToolkit.cpp b/Engine/Plugins/Experimental/MLDeformer/Source/MLDeformerEditor/Private/MLDeformerEditorToolkit.cpp
--- a/Engine/Plugins/Experimental/MLDeformer/Source/MLDeformerEditor/Private/MLDeformerEditorToolkit.cpp
+++ b/Engine/Plugins/Experimental/MLDeformer/Source/MLDeformerEditor/Private/MLDeformerEditorToolkit.cpp
@@ -337,9 +337,7 @@ bool FMLDeformerEditorToolkit::TryLoadOnnxFile() const
 		{
 			EditorData->GetDeformerAsset()->NeuralNetwork = Network;
 			EditorData->GetDeformerAsset()->NeuralNetwork->SetBackEnd(ENeuralBackEnd::Auto);
-			EditorData->GetDeformerAsset()->NeuralNetwork->SetInputDeviceType(ENeuralDeviceType::CPU);
-			EditorData->GetDeformerAsset()->NeuralNetwork->SetDeviceType(ENeuralDeviceType::GPU);
-			EditorData->GetDeformerAsset()->NeuralNetwork->SetOutputDeviceType(ENeuralDeviceType::GPU);
+			EditorData->GetDeformerAsset()->NeuralNetwork->SetDeviceType(/*DeviceType*/ENeuralDeviceType::GPU, /*InputDeviceType*/ENeuralDeviceType::CPU, /*OutputDeviceType*/ENeuralDeviceType::GPU);
 
 			// Recreate the data providers.
 			// This is needed for the neural network GPU buffers to be valid.
diff --git a/Engine/Plugins/Experimental/NeuralNetworkInference/Source/NeuralNetworkInference/Private/NeuralNetwork.cpp b/Engine/Plugins/Experimental/NeuralNetworkInference/Source/NeuralNetworkInference/Private/NeuralNetwork.cpp
--- a/Engine/Plugins/Experimental/NeuralNetworkInference/Source/NeuralNetworkInference/Private/NeuralNetwork.cpp
+++ b/Engine/Plugins/Experimental/NeuralNetworkInference/Source/NeuralNetworkInference/Private/NeuralNetwork.cpp
@@ -128,29 +128,28 @@ ENeuralDeviceType UNeuralNetwork::GetDeviceType() const
 	return DeviceType;
 }
 
-void UNeuralNetwork::SetDeviceType(const ENeuralDeviceType InDeviceType)
-{
-	SetXDeviceType(DeviceType, InDeviceType);
-}
-
 ENeuralDeviceType UNeuralNetwork::GetInputDeviceType() const
 {
 	return InputDeviceType;
 }
 
-void UNeuralNetwork::SetInputDeviceType(const ENeuralDeviceType InInputDeviceType)
-{
-	SetXDeviceType(InputDeviceType, InInputDeviceType);
-}
-
 ENeuralDeviceType UNeuralNetwork::GetOutputDeviceType() const
 {
 	return OutputDeviceType;
 }
 
-void UNeuralNetwork::SetOutputDeviceType(const ENeuralDeviceType InOutputDeviceType)
+void UNeuralNetwork::SetDeviceType(const ENeuralDeviceType InDeviceType, const ENeuralDeviceType InInputDeviceType, const ENeuralDeviceType InOutputDeviceType)
 {
-	SetXDeviceType(OutputDeviceType, InOutputDeviceType);
+	if (DeviceType != InDeviceType || InputDeviceType != InInputDeviceType || OutputDeviceType != InOutputDeviceType)
+	{
+		DeviceType = InDeviceType;
+		InputDeviceType = InInputDeviceType;
+		OutputDeviceType = InOutputDeviceType;
+		if (BackEndForCurrentPlatform == ENeuralBackEnd::UEAndORT)
+		{
+			Load();
+		}
+	}
 }
 
 ENeuralNetworkSynchronousMode UNeuralNetwork::GetSynchronousMode() const
@@ -427,18 +426,6 @@ void UNeuralNetwork::Run()
 /* UNeuralNetwork private functions
  *****************************************************************************/
 
-void UNeuralNetwork::SetXDeviceType(ENeuralDeviceType& OutDeviceType, const ENeuralDeviceType InDeviceType)
-{
-	if (OutDeviceType != InDeviceType)
-	{
-		OutDeviceType = InDeviceType;
-		if (BackEndForCurrentPlatform == ENeuralBackEnd::UEAndORT)
-		{
-			Load();
-		}
-	}
-}
-
 bool UNeuralNetwork::Load()
 {
 	DECLARE_SCOPE_CYCLE_COUNTER(TEXT("UNeuralNetwork_Load"), STAT_UNeuralNetwork_Load, STATGROUP_MachineLearning);
diff --git a/Engine/Plugins/Experimental/NeuralNetworkInference/Source/NeuralNetworkInference/Public/NeuralNetwork.h b/Engine/Plugins/Experimental/NeuralNetworkInference/Source/NeuralNetworkInference/Public/NeuralNetwork.h
--- a/Engine/Plugins/Experimental/NeuralNetworkInference/Source/NeuralNetworkInference/Public/NeuralNetwork.h
+++ b/Engine/Plugins/Experimental/NeuralNetworkInference/Source/NeuralNetworkInference/Public/NeuralNetwork.h
@@ -96,11 +96,9 @@ public:
 	 * @see DeviceType, InputDeviceType, and OutputDeviceType for more details.
 	 */
 	ENeuralDeviceType GetDeviceType() const;
-	void SetDeviceType(const ENeuralDeviceType InDeviceType);
 	ENeuralDeviceType GetInputDeviceType() const;
-	void SetInputDeviceType(const ENeuralDeviceType InInputDeviceType);
 	ENeuralDeviceType GetOutputDeviceType() const;
-	void SetOutputDeviceType(const ENeuralDeviceType InOutputDeviceType);
+	void SetDeviceType(const ENeuralDeviceType InDeviceType, const ENeuralDeviceType InInputDeviceType = ENeuralDeviceType::CPU, const ENeuralDeviceType InOutputDeviceType = ENeuralDeviceType::CPU);
 
 	/**
 	 * Getter and setter functions for SynchronousMode.
@@ -235,11 +233,6 @@ private:
 	struct FImplBackEndUEOnly;
 	TSharedPtr<FImplBackEndUEOnly> ImplBackEndUEOnly;
 
-	/**
-	 * Auxiliary function for Set(Input/Output)DeviceType.
-	 */
-	void SetXDeviceType(ENeuralDeviceType& OutDeviceType, const ENeuralDeviceType InDeviceType);
-
 	/**
 	 * It loads the desired network graph definition and weights internally saved on this UNeuralNetwork instance.
 	 * @return Whether the network was successfully loaded.
diff --git a/Engine/Plugins/Experimental/NeuralNetworkInference/Source/NeuralNetworkInferenceQA/Private/ModelUnitTester.cpp b/Engine/Plugins/Experimental/NeuralNetworkInference/Source/NeuralNetworkInferenceQA/Private/ModelUnitTester.cpp
--- a/Engine/Plugins/Experimental/NeuralNetworkInference/Source/NeuralNetworkInferenceQA/Private/ModelUnitTester.cpp
+++ b/Engine/Plugins/Experimental/NeuralNetworkInference/Source/NeuralNetworkInferenceQA/Private/ModelUnitTester.cpp
@@ -226,9 +226,7 @@ bool FModelUnitTester::ModelAccuracyTest(UNeuralNetwork* InOutNetwork, const ENe
 	{
 		InOutNetwork->SetInputFromArrayCopy(InputArrays[Index]);
 		// CPU
-		InOutNetwork->SetDeviceType(ENeuralDeviceType::CPU);
-		InOutNetwork->SetOutputDeviceType(ENeuralDeviceType::CPU);
-		InOutNetwork->SetInputDeviceType(ENeuralDeviceType::CPU);
+		InOutNetwork->SetDeviceType(/*DeviceType*/ENeuralDeviceType::CPU, /*InputDeviceType*/ENeuralDeviceType::CPU, /*OutputDeviceType*/ENeuralDeviceType::CPU);
 		InOutNetwork->Run();
 		CPUOutputs.Emplace(InOutNetwork->GetOutputTensor().GetArrayCopy<float>());
 	}
@@ -236,9 +234,7 @@ bool FModelUnitTester::ModelAccuracyTest(UNeuralNetwork* InOutNetwork, const ENe
 	{
 		InOutNetwork->SetInputFromArrayCopy(InputArrays[Index]);
 		// Input CPU + GPU + Output CPU
-		InOutNetwork->SetDeviceType(ENeuralDeviceType::GPU);
-		InOutNetwork->SetInputDeviceType(ENeuralDeviceType::CPU);
-		InOutNetwork->SetOutputDeviceType(ENeuralDeviceType::CPU);
+		InOutNetwork->SetDeviceType(/*DeviceType*/ENeuralDeviceType::GPU, /*InputDeviceType*/ENeuralDeviceType::CPU, /*OutputDeviceType*/ENeuralDeviceType::CPU);
 		InOutNetwork->Run();
 		CPUGPUCPUOutputs.Emplace(InOutNetwork->GetOutputTensor().GetArrayCopy<float>());
 	}
@@ -249,15 +245,14 @@ if (InBackEnd == ENeuralBackEnd::UEOnly /*|| true*/)
 {
 		InOutNetwork->SetInputFromArrayCopy(InputArrays[Index]);
 		// Input CPU + GPU + Output GPU
-		InOutNetwork->SetDeviceType(ENeuralDeviceType::GPU);
-		InOutNetwork->SetInputDeviceType(ENeuralDeviceType::CPU);
-		InOutNetwork->SetOutputDeviceType(ENeuralDeviceType::GPU);
+		InOutNetwork->SetDeviceType(/*DeviceType*/ENeuralDeviceType::GPU, /*InputDeviceType*/ENeuralDeviceType::CPU, /*OutputDeviceType*/ENeuralDeviceType::GPU);
 		InOutNetwork->Run();
 		InOutNetwork->OutputTensorsToCPU();
 		CPUGPUGPUOutputs.Emplace(InOutNetwork->GetOutputTensor().GetArrayCopy<float>());
 }
 else
 {
+	UE_LOG(LogNeuralNetworkInferenceQA, Display, TEXT("FModelUnitTester::ModelAccuracyTest(): GPU output for UEAndORT back end not working yet. Uncomment this line to test it."));
 	CPUGPUGPUOutputs.Push(CPUGPUCPUOutputs[Index]);
 }
 	}
@@ -305,9 +300,7 @@ else
 		bDidGlobalTestPassed &= !bDidSomeTestFailed;
 	}
 	// Reset to original network state
-	InOutNetwork->SetDeviceType(OriginalDeviceType);
-	InOutNetwork->SetInputDeviceType(OriginalInputDeviceType);
-	InOutNetwork->SetOutputDeviceType(OriginalOutputDeviceType);
+	InOutNetwork->SetDeviceType(/*DeviceType*/OriginalDeviceType, /*InputDeviceType*/OriginalInputDeviceType, /*OutputDeviceType*/OriginalOutputDeviceType);
 	InOutNetwork->SetBackEnd(OriginalBackEnd);
 	// Test successful
 	return bDidGlobalTestPassed;
