NeuralNetworkInference (NNI) plugin:
- Expose DX12 compatibilitiy functions so unit test can be avoid with a Display message if so.

_codereview Francisco.Vicente
#rnx

diff --git a/Engine/Plugins/Experimental/NeuralNetworkInference/Source/NeuralNetworkInference/Private/NeuralNetwork.cpp b/Engine/Plugins/Experimental/NeuralNetworkInference/Source/NeuralNetworkInference/Private/NeuralNetwork.cpp
--- a/Engine/Plugins/Experimental/NeuralNetworkInference/Source/NeuralNetworkInference/Private/NeuralNetwork.cpp
+++ b/Engine/Plugins/Experimental/NeuralNetworkInference/Source/NeuralNetworkInference/Private/NeuralNetwork.cpp
@@ -145,7 +145,7 @@ void UNeuralNetwork::SetDeviceType(const ENeuralDeviceType InDeviceType, const E
 		DeviceType = InDeviceType;
 		InputDeviceType = InInputDeviceType;
 		OutputDeviceType = InOutputDeviceType;
-		if (BackEndForCurrentPlatform == ENeuralBackEnd::UEAndORT)
+		if (bIsLoaded && BackEndForCurrentPlatform == ENeuralBackEnd::UEAndORT) // No need to re-load if not bIsLoaded
 		{
 			Load();
 		}
@@ -186,8 +186,35 @@ void UNeuralNetwork::SetBackEnd(const ENeuralBackEnd InBackEnd)
 	if (BackEndForCurrentPlatform != NewBackEndForCurrentPlatform)
 	{
 		BackEndForCurrentPlatform = NewBackEndForCurrentPlatform;
-		Load();
+		if (bIsLoaded) // No need to re-load if not bIsLoaded
+		{
+			Load();
+		}
+	}
+}
+
+bool UNeuralNetwork::IsGPUConfigCompatibleForUEAndORTBackEnd()
+{
+	return FImplBackEndUEAndORT::IsGPUConfigCompatible();
+}
+
+bool UNeuralNetwork::IsGPUConfigCompatibleForCurrentBackEnd() const
+{
+	// UEAndORT
+	if (BackEndForCurrentPlatform == ENeuralBackEnd::UEAndORT)
+	{
+		return IsGPUConfigCompatibleForUEAndORTBackEnd();
 	}
+
+	// UEOnly
+	else if (BackEndForCurrentPlatform == ENeuralBackEnd::UEOnly)
+	{
+		return true;
+	}
+
+	// Unknown
+	UE_LOG(LogNeuralNetworkInference, Warning, TEXT("UNeuralNetwork::IsGPUSupported(): Unknown [BackEnd,BackEndForCurrentPlatform] = [%d,%d]."), (int32)BackEnd, (int32)BackEndForCurrentPlatform);
+	return false;
 }
 
 const FNeuralTensor& UNeuralNetwork::GetInputTensor(const int32 InTensorIndex) const
@@ -436,7 +463,7 @@ bool UNeuralNetwork::Load()
 	// UEAndORT
 	if (BackEndForCurrentPlatform == ENeuralBackEnd::UEAndORT)
 	{
-		UNeuralNetwork::FImplBackEndUEAndORT::SetDeviceToCPUIfD3D12NotSupported(DeviceType);
+		UNeuralNetwork::FImplBackEndUEAndORT::WarnAndSetDeviceToCPUIfDX12NotEnabled(DeviceType);
 		bIsLoaded = UNeuralNetwork::FImplBackEndUEAndORT::Load(ImplBackEndUEAndORT, AreInputTensorSizesVariable, ModelReadFromFileInBytes, ModelFullFilePath, GetDeviceType(), GetInputDeviceType(), GetOutputDeviceType());
 	}
 	// UEOnly
@@ -502,6 +529,13 @@ void UNeuralNetwork::PostLoad()
 	// If ModelReadFromFileInBytes is not empty, call Load() 
 	if (ModelReadFromFileInBytes.Num() > 0)
 	{
+		// If GPU selected but not compatible, set to CPU
+		if (!IsGPUConfigCompatibleForCurrentBackEnd())
+		{
+			UE_LOG(LogNeuralNetworkInference, Warning, TEXT("UNeuralNetwork::PostLoad(): IsGPUConfigCompatibleForCurrentBackEnd() returned false, setting DeviceType to CPU."));
+			DeviceType = ENeuralDeviceType::CPU;
+		}
+		// Load
 		if (!Load())
 		{
 			UE_LOG(LogNeuralNetworkInference, Warning, TEXT("UNeuralNetwork::PostLoad(): UNeuralNetwork could not be loaded."));
diff --git a/Engine/Plugins/Experimental/NeuralNetworkInference/Source/NeuralNetworkInference/Private/NeuralNetworkImplBackEndUEAndORT.cpp b/Engine/Plugins/Experimental/NeuralNetworkInference/Source/NeuralNetworkInference/Private/NeuralNetworkImplBackEndUEAndORT.cpp
--- a/Engine/Plugins/Experimental/NeuralNetworkInference/Source/NeuralNetworkInference/Private/NeuralNetworkImplBackEndUEAndORT.cpp
+++ b/Engine/Plugins/Experimental/NeuralNetworkInference/Source/NeuralNetworkInference/Private/NeuralNetworkImplBackEndUEAndORT.cpp
@@ -123,56 +123,44 @@ IDMLDevice* FPrivateImplBackEndUEAndORT::FDMLDeviceList::Add(ID3D12Device* Devic
 /* UNeuralNetwork public functions
  *****************************************************************************/
 
-void UNeuralNetwork::FImplBackEndUEAndORT::SetDeviceToCPUIfD3D12NotSupported(ENeuralDeviceType& InOutDeviceType)
+void UNeuralNetwork::FImplBackEndUEAndORT::WarnAndSetDeviceToCPUIfDX12NotEnabled(ENeuralDeviceType& InOutDeviceType)
 {
 	if (InOutDeviceType != ENeuralDeviceType::CPU)
 	{
-		if (!IsD3D12Supported())
+		if (!IsGPUConfigCompatible())
 		{
-			UE_LOG(LogNeuralNetworkInference, Warning, TEXT("FImplBackEndUEAndORT::SetDeviceToCPUIfD3D12NotSupported(): Setting device to CPU provisionally."));
 			InOutDeviceType = ENeuralDeviceType::CPU;
-		}
-	}
-}
 
-bool UNeuralNetwork::FImplBackEndUEAndORT::IsD3D12Supported()
-{
-#ifdef WITH_UE_AND_ORT_SUPPORT
-
-#ifdef PLATFORM_WIN64
-	// To create a DirectML device we need to check that we're using DX12 first
-	const FString RHIName = GDynamicRHI->GetName();
-
-	// Display warning if not DX12
-	if (RHIName != TEXT("D3D12"))
-	{
-		const FString ErrorMessage = TEXT("On Windows, only DirectX 12 rendering (\"D3D12\") is compatible with NeuralNetworkInference (NNI). Instead, \"") + RHIName
-			+ TEXT("\" was used. You have the following options:\n"
-				"\t1. (Recommended) Switch Unreal Engine to DX12. In order to do that:\n"
+			const FString RHIName = GDynamicRHI->GetName();
+			const FString ErrorMessage = TEXT("On Windows, only DirectX 12 rendering (\"D3D12\") is compatible with the UEAndORT back end of NeuralNetworkInference (NNI). Instead, \"") + RHIName
+				+ TEXT("\" was used. You have the following options:\n\n"
+					"\t1. (Recommended) Switch Unreal Engine to DX12. In order to do that:\n"
 					"\t\t - Go to \"Project Settings\", \"Platforms\", \"Windows\", \"Default RHI\".\n"
 					"\t\t - Select \"DirectX 12\".\n"
 					"\t\t - Restart Unreal Engine.\n"
-				"\t2. Alternatively, switch the network to CPU with UNeuralNetwork::SetDeviceType().\n"
-				"\t3. (Not recommended) You could also switch the network to UEOnly with UNeuralNetwork::SetBackEnd().\n"
-				"Setting the network to CPU provisionally.");
-		UE_LOG(LogNeuralNetworkInference, Warning, TEXT("FImplBackEndUEAndORT::ConfigureMembers(): %s"), *ErrorMessage, *RHIName);
+					"\t2. Alternatively, switch the network to CPU with UNeuralNetwork::SetDeviceType().\n"
+					"\t3. (Not recommended) You could also switch the network to UEOnly with UNeuralNetwork::SetBackEnd().\n\n"
+					"Network set to CPU provisionally.");
+			UE_LOG(LogNeuralNetworkInference, Warning, TEXT("FImplBackEndUEAndORT::WarnAndSetDeviceToCPUIfDX12NotEnabled(): %s"), *ErrorMessage);
 #if WITH_EDITOR
-		FMessageDialog::Open(EAppMsgType::Ok, FText::FromString(ErrorMessage));
+			FMessageDialog::Open(EAppMsgType::Ok, FText::FromString(ErrorMessage));
 #endif //WITH_EDITOR
-		return false;
+		}
 	}
+}
 
-	return true;
-
-#else //PLATFORM_WIN64
-	UE_LOG(LogNeuralNetworkInference, Warning, TEXT("FImplBackEndUEAndORT::IsD3D12Supported(): This function will always return false outside of Windows."));
-	return false;
+bool UNeuralNetwork::FImplBackEndUEAndORT::IsGPUConfigCompatible()
+{
+#ifdef WITH_UE_AND_ORT_SUPPORT
+#ifdef PLATFORM_WIN64
+	// Return whether it is DX12
+	const FString RHIName = GDynamicRHI->GetName();
+	return (RHIName == TEXT("D3D12"));
 #endif //PLATFORM_WIN64
-
-#else //WITH_UE_AND_ORT_SUPPORT
-	UE_LOG(LogNeuralNetworkInference, Warning, TEXT("FImplBackEndUEAndORT::IsD3D12Supported(): This function will always return false if WITH_UE_AND_ORT_SUPPORT is not defined."));
-	return false;
 #endif //WITH_UE_AND_ORT_SUPPORT
+
+	// If not Windows and/or if WITH_UE_AND_ORT_SUPPORT not defined, then this should return true because GPU will always work
+	return true;
 }
 
 //bool UNeuralNetwork::FImplBackEndUEAndORT::Load(TSharedPtr<FImplBackEndUEAndORT>& InOutImplBackEndUEAndORT, TArray<bool>& OutAreInputTensorSizesVariable, const TArray<uint8>& InModelReadFromFileInBytes, const FString& InModelFullFilePath, const ENeuralDeviceType InDeviceType)
@@ -433,8 +421,10 @@ bool UNeuralNetwork::FImplBackEndUEAndORT::ConfigureMembers(const ENeuralDeviceT
 	if (InDeviceType == ENeuralDeviceType::GPU)
 	{
 #ifdef PLATFORM_WIN64
-		if (!IsD3D12Supported())
+		// To create a DirectML device we need to check that we're using DX12 first
+		if (!IsGPUConfigCompatible())
 		{
+			UE_LOG(LogNeuralNetworkInference, Warning, TEXT("FImplBackEndUEAndORT::ConfigureMembers(): UEAndORT back end for GPU needs DX12 enabled."));
 			return false;
 		}
 
diff --git a/Engine/Plugins/Experimental/NeuralNetworkInference/Source/NeuralNetworkInference/Private/NeuralNetworkImplBackEndUEAndORT.h b/Engine/Plugins/Experimental/NeuralNetworkInference/Source/NeuralNetworkInference/Private/NeuralNetworkImplBackEndUEAndORT.h
--- a/Engine/Plugins/Experimental/NeuralNetworkInference/Source/NeuralNetworkInference/Private/NeuralNetworkImplBackEndUEAndORT.h
+++ b/Engine/Plugins/Experimental/NeuralNetworkInference/Source/NeuralNetworkInference/Private/NeuralNetworkImplBackEndUEAndORT.h
@@ -49,9 +49,9 @@ public:
 	TArray<const char*> OutputTensorNames; /* Tensor names */
 #endif //WITH_UE_AND_ORT_SUPPORT
 
-	static void SetDeviceToCPUIfD3D12NotSupported(ENeuralDeviceType& InOutDeviceType);
+	static void WarnAndSetDeviceToCPUIfDX12NotEnabled(ENeuralDeviceType& InOutDeviceType);
 
-	static bool IsD3D12Supported();
+	static bool IsGPUConfigCompatible();
 
 	static bool Load(TSharedPtr<FImplBackEndUEAndORT>& InOutImplBackEndUEAndORT, TArray<bool>& OutAreInputTensorSizesVariable, const TArray<uint8>& InModelReadFromFileInBytes, const FString& InModelFullFilePath, const ENeuralDeviceType InDeviceType, const ENeuralDeviceType InInputDeviceType, const ENeuralDeviceType InOutputDeviceType);
 	
diff --git a/Engine/Plugins/Experimental/NeuralNetworkInference/Source/NeuralNetworkInference/Public/NeuralNetwork.h b/Engine/Plugins/Experimental/NeuralNetworkInference/Source/NeuralNetworkInference/Public/NeuralNetwork.h
--- a/Engine/Plugins/Experimental/NeuralNetworkInference/Source/NeuralNetworkInference/Public/NeuralNetwork.h
+++ b/Engine/Plugins/Experimental/NeuralNetworkInference/Source/NeuralNetworkInference/Public/NeuralNetwork.h
@@ -127,6 +127,14 @@ public:
 	ENeuralBackEnd GetBackEnd() const;
 	ENeuralBackEnd GetBackEndForCurrentPlatform() const;
 	void SetBackEnd(const ENeuralBackEnd InBackEnd);
+	
+	/**
+	 * IsGPUConfigCompatible will always return true for ENeuralBackEnd::UEOnly. For ENeuralBackEnd::UEAndORT, IsGPUConfigCompatible() and IsGPUConfigCompatibleForUEAndORTBackEnd() will return the same value:
+	 * - True if DX12 is enabled, meaning UEAndORT can run on both the CPU and GPU. Also true if the current platform is not Windows.
+	 * - False if DX12 is disabled, meaning UEAndORT can run only run on the CPU. The user will need to enable DX12 to be able to run GPU, switch to CPU, or switch to the UEOnly back end.
+	 */
+	static bool IsGPUConfigCompatibleForUEAndORTBackEnd();
+	bool IsGPUConfigCompatibleForCurrentBackEnd() const;
 
 	/**
 	 * Functions to get/fill input.
