NeuralNetworkInference (NNI) plugin:
- Added unit test for GPU input and required auxiliary functions for UNN for input tensors.

_codereview Francisco.Vicente
#rnx

diff --git a/Engine/Plugins/Experimental/NeuralNetworkInference/Source/NeuralNetworkInference/Private/NeuralNetwork.cpp b/Engine/Plugins/Experimental/NeuralNetworkInference/Source/NeuralNetworkInference/Private/NeuralNetwork.cpp
--- a/Engine/Plugins/Experimental/NeuralNetworkInference/Source/NeuralNetworkInference/Private/NeuralNetwork.cpp
+++ b/Engine/Plugins/Experimental/NeuralNetworkInference/Source/NeuralNetworkInference/Private/NeuralNetwork.cpp
@@ -217,27 +217,27 @@ bool UNeuralNetwork::IsGPUConfigCompatibleForCurrentBackEnd() const
 	return false;
 }
 
+#define GET_TENSOR_CODE(InTensorIndex, UEAndORTOnly_GetIndexes, UEOnly_GetTensorsFunction, UEOnly_Tensors) \
+	/* Sanity check */ \
+	checkf(bIsLoaded, TEXT("Call UNeuralNetwork::Load() to load a model first.")); \
+	/* UEAndORT */ \
+	if (BackEndForCurrentPlatform == ENeuralBackEnd::UEAndORT) \
+	{ \
+		return ImplBackEndUEAndORT->UEOnly_Tensors[InTensorIndex]; \
+	} \
+	/* UEOnly */ \
+	else if (BackEndForCurrentPlatform == ENeuralBackEnd::UEOnly) \
+	{ \
+		FNeuralTensorManager& TensorManager = ImplBackEndUEOnly->TensorManager; \
+		return TensorManager.UEOnly_GetTensorsFunction()[TensorManager.UEAndORTOnly_GetIndexes()[InTensorIndex]]; \
+	} \
+	/* Unknown */ \
+	checkf(false, TEXT("Unknown [BackEnd,BackEndForCurrentPlatform] = [%d,%d]."), (int32)BackEnd, (int32)BackEndForCurrentPlatform); \
+	return ImplBackEndUEAndORT->UEOnly_Tensors[InTensorIndex]
+
 const FNeuralTensor& UNeuralNetwork::GetInputTensor(const int32 InTensorIndex) const
 {
-	// Sanity check
-	checkf(bIsLoaded, TEXT("Call UNeuralNetwork::Load() to load a model first."));
-
-	// UEAndORT
-	if (BackEndForCurrentPlatform == ENeuralBackEnd::UEAndORT)
-	{
-		return ImplBackEndUEAndORT->InputTensors[InTensorIndex];
-	}
-
-	// UEOnly
-	else if (BackEndForCurrentPlatform == ENeuralBackEnd::UEOnly)
-	{
-		FNeuralTensorManager& TensorManager = ImplBackEndUEOnly->TensorManager;
-		return TensorManager.GetTensors()[TensorManager.GetInputIndexes()[InTensorIndex]];
-	}
-
-	// Unknown
-	checkf(false, TEXT("UNeuralNetwork::GetInputTensor(): Unknown [BackEnd,BackEndForCurrentPlatform] = [%d,%d]."), (int32)BackEnd, (int32)BackEndForCurrentPlatform);
-	return ImplBackEndUEAndORT->InputTensors[InTensorIndex];
+	GET_TENSOR_CODE(InTensorIndex, GetInputIndexes, GetTensors, InputTensors);
 }
 
 void UNeuralNetwork::SetInputFromArrayCopy(const TArray<float>& InArray, const int32 InTensorIndex)
@@ -321,32 +321,9 @@ int64 UNeuralNetwork::GetInputTensorNumber() const
 	return -1;
 }
 
-#define GET_OUTPUT_TENSOR_CODE(InTensorIndex, UEOnly_GetTensorsFunction) \
-	/* Sanity check */ \
-	checkf(bIsLoaded, TEXT("Call UNeuralNetwork::Load() to load a model first.")); \
-	/* UEAndORT */ \
-	if (BackEndForCurrentPlatform == ENeuralBackEnd::UEAndORT) \
-	{ \
-		return ImplBackEndUEAndORT->OutputTensors[InTensorIndex]; \
-	} \
-	/* UEOnly */ \
-	else if (BackEndForCurrentPlatform == ENeuralBackEnd::UEOnly) \
-	{ \
-		FNeuralTensorManager& TensorManager = ImplBackEndUEOnly->TensorManager; \
-		return TensorManager.UEOnly_GetTensorsFunction()[TensorManager.GetOutputIndexes()[InTensorIndex]]; \
-	} \
-	/* Unknown */ \
-	checkf(false, TEXT("UNeuralNetwork::GetOutputTensor(): Unknown [BackEnd,BackEndForCurrentPlatform] = [%d,%d]."), (int32)BackEnd, (int32)BackEndForCurrentPlatform); \
-	return ImplBackEndUEAndORT->OutputTensors[InTensorIndex]
-
 const FNeuralTensor& UNeuralNetwork::GetOutputTensor(const int32 InTensorIndex) const
 {
-	GET_OUTPUT_TENSOR_CODE(InTensorIndex, GetTensors);
-}
-
-FNeuralTensor& UNeuralNetwork::GetOutputTensorMutable(const int32 InTensorIndex)
-{
-	GET_OUTPUT_TENSOR_CODE(InTensorIndex, GetTensorsMutable);
+	GET_TENSOR_CODE(InTensorIndex, GetOutputIndexes, GetTensors, OutputTensors);
 }
 
 int64 UNeuralNetwork::GetOutputTensorNumber() const
@@ -375,7 +352,48 @@ int64 UNeuralNetwork::GetOutputTensorNumber() const
 	return -1;
 }
 
-void UNeuralNetwork::OutputTensorsToCPU(const TArray<int32>& InOutputTensorIndexes)
+void UNeuralNetwork::InputTensorsToGPU(const TArray<int32>& InTensorIndexes)
+{
+	// Sanity check
+	if (!bIsLoaded)
+	{
+		UE_LOG(LogNeuralNetworkInference, Warning, TEXT("UNeuralNetwork::InputTensorsToGPU(): Call UNeuralNetwork::Load() to load a model first."));
+		return;
+	}
+
+	// In RHI
+	ENQUEUE_RENDER_COMMAND(UNeuralNetwork_InputTensorToGPU_RenderThread)(
+		[this, &InTensorIndexes](FRHICommandListImmediate& RHICmdList)
+		{
+			FRDGBuilder GraphBuilder(RHICmdList, RDG_EVENT_NAME("UNeuralNetwork::InputTensorToGPU()"));
+			// Run for each input tensor
+			if (InTensorIndexes.IsEmpty())
+			{
+				// Refresh tensor(s) w.r.t. GraphBuilder + Move memory from CPU to GPU
+				for (uint32 InputTensorIndex = 0; InputTensorIndex < GetInputTensorNumber(); ++InputTensorIndex)
+				{
+					FNeuralTensor& InputTensor = GetInputTensorMutable(InputTensorIndex);
+					InputTensor.ToGPU_RenderThread(&GraphBuilder);
+				}
+			}
+			// Run for the desired input tensors
+			else
+			{
+				for (const int32 InputTensorIndex : InTensorIndexes)
+				{
+					FNeuralTensor& InputTensor = GetInputTensorMutable(InputTensorIndex);
+					InputTensor.ToGPU_RenderThread(&GraphBuilder);
+				}
+			}
+			// Execute Render Graph
+			GraphBuilder.Execute();
+		}
+	);
+
+	FNeuralNetworkInferenceUtils::WaitUntilRHIFinished();
+}
+
+void UNeuralNetwork::OutputTensorsToCPU(const TArray<int32>& InTensorIndexes)
 {
 	// Sanity check
 	if (!bIsLoaded)
@@ -386,11 +404,11 @@ void UNeuralNetwork::OutputTensorsToCPU(const TArray<int32>& InOutputTensorIndex
 
 	// In RHI
 	ENQUEUE_RENDER_COMMAND(UNeuralNetwork_OutputTensorToCPU_RenderThread)(
-		[this, &InOutputTensorIndexes](FRHICommandListImmediate& RHICmdList)
+		[this, &InTensorIndexes](FRHICommandListImmediate& RHICmdList)
 		{
 			FRDGBuilder GraphBuilder(RHICmdList, RDG_EVENT_NAME("UNeuralNetwork::OutputTensorToCPU()"));
 			// Run for each output tensor
-			if (InOutputTensorIndexes.IsEmpty())
+			if (InTensorIndexes.IsEmpty())
 			{
 				// Refresh tensor(s) w.r.t. GraphBuilder + Move memory from GPU to CPU
 				for (uint32 OutputTensorIndex = 0; OutputTensorIndex < GetOutputTensorNumber(); ++OutputTensorIndex)
@@ -403,7 +421,7 @@ void UNeuralNetwork::OutputTensorsToCPU(const TArray<int32>& InOutputTensorIndex
 			// Run for the desired output tensors
 			else
 			{
-				for (const int32 OutputTensorIndex : InOutputTensorIndexes)
+				for (const int32 OutputTensorIndex : InTensorIndexes)
 				{
 					FNeuralTensor& OutputTensor = GetOutputTensorMutable(OutputTensorIndex);
 					OutputTensor.UpdateSRVAndOrUAV_RenderThread(&GraphBuilder);
@@ -480,6 +498,16 @@ bool UNeuralNetwork::Load()
 	return bIsLoaded;
 }
 
+FNeuralTensor& UNeuralNetwork::GetInputTensorMutable(const int32 InTensorIndex)
+{
+	GET_TENSOR_CODE(InTensorIndex, GetInputIndexes, GetTensorsMutable, InputTensors);
+}
+
+FNeuralTensor& UNeuralNetwork::GetOutputTensorMutable(const int32 InTensorIndex)
+{
+	GET_TENSOR_CODE(InTensorIndex, GetOutputIndexes, GetTensorsMutable, OutputTensors);
+}
+
 #if WITH_EDITOR
 UAssetImportData* UNeuralNetwork::GetAssetImportData() const
 {
diff --git a/Engine/Plugins/Experimental/NeuralNetworkInference/Source/NeuralNetworkInference/Public/NeuralNetwork.h b/Engine/Plugins/Experimental/NeuralNetworkInference/Source/NeuralNetworkInference/Public/NeuralNetwork.h
--- a/Engine/Plugins/Experimental/NeuralNetworkInference/Source/NeuralNetworkInference/Public/NeuralNetwork.h
+++ b/Engine/Plugins/Experimental/NeuralNetworkInference/Source/NeuralNetworkInference/Public/NeuralNetwork.h
@@ -151,11 +151,13 @@ public:
 	const FNeuralTensor& GetOutputTensor(const int32 InTensorIndex = 0) const;
 	int64 GetOutputTensorNumber() const;
 	/**
-	 * Non-efficient function meant to be used only for debugging purposes. It will send the GPU memory of the desired output tensor(s) back into CPU memory.
-	 * It is useful for debugging OutputDeviceType == ENeuralDeviceType::GPU cases.
-	 * @param InOutputTensorIndexes If empty (default value), it will apply to all output tensors.
+	 * Non-efficient functions meant to be used only for debugging purposes.
+	 * - InputTensorsToCPU will send the CPU memory of the desired input tensor(s) to GPU memory. Used to debug InputDeviceType == ENeuralDeviceType::GPU.
+	 * - OutputTensorsToCPU will send the GPU memory of the desired output tensor(s) back to CPU memory. Used to debug OutputDeviceType == ENeuralDeviceType::GPU.
+	 * @param InTensorIndexes If empty (default value), it will apply to all output tensors.
 	 */
-	void OutputTensorsToCPU(const TArray<int32>& InOutputTensorIndexes = TArray<int32>());
+	void InputTensorsToGPU(const TArray<int32>& InTensorIndexes = TArray<int32>());
+	void OutputTensorsToCPU(const TArray<int32>& InTensorIndexes = TArray<int32>());
 
 	/**
 	 * Run() executes the forward pass on the current UNeuralNetwork given the current input FDeprecatedNeuralTensor(s), which were previously filled with
@@ -248,8 +250,9 @@ private:
 	bool Load();
 
 	/**
-	 * Mutable version of GetOutputTensor(const int32 InTensorIndex) const.
+	 * Private and mutable version of GetInputTensor()/GetOutputTensor().
 	 */
+	FNeuralTensor& GetInputTensorMutable(const int32 InTensorIndex = 0);
 	FNeuralTensor& GetOutputTensorMutable(const int32 InTensorIndex = 0);
 
 public:
diff --git a/Engine/Plugins/Experimental/NeuralNetworkInference/Source/NeuralNetworkInferenceQA/Private/ModelUnitTester.cpp b/Engine/Plugins/Experimental/NeuralNetworkInference/Source/NeuralNetworkInferenceQA/Private/ModelUnitTester.cpp
--- a/Engine/Plugins/Experimental/NeuralNetworkInference/Source/NeuralNetworkInferenceQA/Private/ModelUnitTester.cpp
+++ b/Engine/Plugins/Experimental/NeuralNetworkInference/Source/NeuralNetworkInferenceQA/Private/ModelUnitTester.cpp
@@ -222,6 +222,7 @@ bool FModelUnitTester::ModelAccuracyTest(UNeuralNetwork* InOutNetwork, const ENe
 	// Multiple for loops to make sure that running on the CPU does not affect GPU results or vice versa
 	bool bDidGlobalTestPassed = true;
 	TArray<TArray<float>> CPUOutputs, CPUGPUCPUOutputs, CPUGPUGPUOutputs;
+	// Input CPU + Network CPU + Output CPU
 	for (int32 Index = 0; Index < InputArrays.Num(); ++Index)
 	{
 		InOutNetwork->SetInputFromArrayCopy(InputArrays[Index]);
@@ -230,27 +231,45 @@ bool FModelUnitTester::ModelAccuracyTest(UNeuralNetwork* InOutNetwork, const ENe
 		InOutNetwork->Run();
 		CPUOutputs.Emplace(InOutNetwork->GetOutputTensor().GetArrayCopy<float>());
 	}
+	// Input CPU + Network GPU + Output CPU
 	for (int32 Index = 0; Index < InputArrays.Num(); ++Index)
 	{
 		InOutNetwork->SetInputFromArrayCopy(InputArrays[Index]);
-		// Input CPU + GPU + Output CPU
 		InOutNetwork->SetDeviceType(/*DeviceType*/ENeuralDeviceType::GPU, /*InputDeviceType*/ENeuralDeviceType::CPU, /*OutputDeviceType*/ENeuralDeviceType::CPU);
 		InOutNetwork->Run();
 		CPUGPUCPUOutputs.Emplace(InOutNetwork->GetOutputTensor().GetArrayCopy<float>());
 	}
+	// Input CPU + Network GPU + Output GPU
 	for (int32 Index = 0; Index < InputArrays.Num(); ++Index)
 	{
 // @todo: Remove if/else (and else contents) once DX12 version of NNI (UEAndORT) is ready
 if (InBackEnd == ENeuralBackEnd::UEOnly /*|| true*/)
 {
 		InOutNetwork->SetInputFromArrayCopy(InputArrays[Index]);
-		// Input CPU + GPU + Output GPU
 		InOutNetwork->SetDeviceType(/*DeviceType*/ENeuralDeviceType::GPU, /*InputDeviceType*/ENeuralDeviceType::CPU, /*OutputDeviceType*/ENeuralDeviceType::GPU);
 		InOutNetwork->Run();
 		InOutNetwork->OutputTensorsToCPU();
 		CPUGPUGPUOutputs.Emplace(InOutNetwork->GetOutputTensor().GetArrayCopy<float>());
 }
 else
+{
+	UE_LOG(LogNeuralNetworkInferenceQA, Display, TEXT("FModelUnitTester::ModelAccuracyTest(): GPU output for UEAndORT back end not working yet. Uncomment this line to test it."));
+	CPUGPUGPUOutputs.Push(CPUGPUCPUOutputs[Index]);
+}
+	}
+	// Input GPU + Network GPU + Output CPU
+	for (int32 Index = 0; Index < InputArrays.Num(); ++Index)
+	{
+// @todo: Remove if/else (and else contents) once DX12 version of NNI (UEAndORT) is ready
+if (InBackEnd == ENeuralBackEnd::UEOnly /*|| true*/)
+{
+		InOutNetwork->SetInputFromArrayCopy(InputArrays[Index]);
+		InOutNetwork->InputTensorsToGPU();
+		InOutNetwork->SetDeviceType(/*DeviceType*/ENeuralDeviceType::GPU, /*InputDeviceType*/ENeuralDeviceType::GPU, /*OutputDeviceType*/ENeuralDeviceType::CPU);
+		InOutNetwork->Run();
+		CPUGPUGPUOutputs.Emplace(InOutNetwork->GetOutputTensor().GetArrayCopy<float>());
+}
+else
 {
 	UE_LOG(LogNeuralNetworkInferenceQA, Display, TEXT("FModelUnitTester::ModelAccuracyTest(): GPU output for UEAndORT back end not working yet. Uncomment this line to test it."));
 	CPUGPUGPUOutputs.Push(CPUGPUCPUOutputs[Index]);
