NeuralNetworkInference (NNI) plugin:
- Fix and re-enable unit test to pass if DX11 used (by skipping DX12 tests).
- Fix GPU input test not actually being used.

_codereview Francisco.Vicente
#rnx

diff --git a/Engine/Plugins/Experimental/NeuralNetworkInference/Source/NeuralNetworkInferenceQA/Private/ModelUnitTester.cpp b/Engine/Plugins/Experimental/NeuralNetworkInference/Source/NeuralNetworkInferenceQA/Private/ModelUnitTester.cpp
--- a/Engine/Plugins/Experimental/NeuralNetworkInference/Source/NeuralNetworkInferenceQA/Private/ModelUnitTester.cpp
+++ b/Engine/Plugins/Experimental/NeuralNetworkInference/Source/NeuralNetworkInferenceQA/Private/ModelUnitTester.cpp
@@ -123,7 +123,14 @@ bool FModelUnitTester::ModelLoadAccuracyAndSpeedTests(const FString& InProjectCo
 #ifdef WITH_UE_AND_ORT_SUPPORT
 		bDidGlobalTestPassed &= ModelSpeedTest(UAssetModelFilePath, ENeuralDeviceType::CPU, ENeuralBackEnd::UEAndORT, InCPURepetitionsForUEAndORTBackEnd[ModelIndex]);
 #ifdef PLATFORM_WIN64
-		bDidGlobalTestPassed &= ModelSpeedTest(UAssetModelFilePath, ENeuralDeviceType::GPU, ENeuralBackEnd::UEAndORT, InGPURepetitionsForUEAndORTBackEnd[ModelIndex]);
+		if (UNeuralNetwork::IsGPUConfigCompatibleForUEAndORTBackEnd())
+		{
+			bDidGlobalTestPassed &= ModelSpeedTest(UAssetModelFilePath, ENeuralDeviceType::GPU, ENeuralBackEnd::UEAndORT, InGPURepetitionsForUEAndORTBackEnd[ModelIndex]);
+		}
+		else
+		{
+			UE_LOG(LogNeuralNetworkInferenceQA, Display, TEXT("-------------------- ModelSpeedTest-UEAndORT-GPU skipped (DX12 not enabled)."));
+		}
 #else //PLATFORM_WIN64
 		bDidGlobalTestPassed &= ModelSpeedTest(UAssetModelFilePath, ENeuralDeviceType::GPU, ENeuralBackEnd::UEAndORT, 0);
 #endif //PLATFORM_WIN64
@@ -160,6 +167,11 @@ UNeuralNetwork* FModelUnitTester::NetworkUassetLoadTest(const FString& InUAssetP
 		ensureMsgf(false, TEXT("UNeuralNetwork is a nullptr. Path: \"%s\"."), *InUAssetPath);
 		return nullptr;
 	}
+	if (!UNeuralNetwork::IsGPUConfigCompatibleForUEAndORTBackEnd())
+	{
+		UE_LOG(LogNeuralNetworkInferenceQA, Display, TEXT("NetworkUassetLoadTest-UEAndORT: Set to CPU (DX12 not enabled)."));
+		Network->SetDeviceType(ENeuralDeviceType::CPU);
+	}
 	if (!Network->IsLoaded())
 	{
 		ensureMsgf(false, TEXT("UNeuralNetwork could not be loaded from uasset disk location. Path: \"%s\"."), *InUAssetPath);
@@ -177,6 +189,11 @@ UNeuralNetwork* FModelUnitTester::NetworkONNXOrORTLoadTest(const FString& InMode
 		ensureMsgf(false, TEXT("UNeuralNetwork is a nullptr. Path: \"%s\"."), *InModelFilePath);
 		return nullptr;
 	}
+	if (!UNeuralNetwork::IsGPUConfigCompatibleForUEAndORTBackEnd())
+	{
+		UE_LOG(LogNeuralNetworkInferenceQA, Display, TEXT("NetworkONNXOrORTLoadTest-UEAndORT: Set to CPU (DX12 not enabled)."));
+		Network->SetDeviceType(ENeuralDeviceType::CPU);
+	}
 	if (!Network->Load(InModelFilePath))
 	{
 		ensureMsgf(false, TEXT("UNeuralNetwork could not be loaded from ONNX file disk location. Path: \"%s\"."), *InModelFilePath);
@@ -221,7 +238,7 @@ bool FModelUnitTester::ModelAccuracyTest(UNeuralNetwork* InOutNetwork, const ENe
 	// Run each input with CPU/GPU and compare with each other and with the ground truth
 	// Multiple for loops to make sure that running on the CPU does not affect GPU results or vice versa
 	bool bDidGlobalTestPassed = true;
-	TArray<TArray<float>> CPUOutputs, CPUGPUCPUOutputs, CPUGPUGPUOutputs;
+	TArray<TArray<float>> CPUOutputs, CPUGPUCPUOutputs, CPUGPUGPUOutputs, GPUGPUCPUOutputs;
 	// Input CPU + Network CPU + Output CPU
 	for (int32 Index = 0; Index < InputArrays.Num(); ++Index)
 	{
@@ -231,71 +248,84 @@ bool FModelUnitTester::ModelAccuracyTest(UNeuralNetwork* InOutNetwork, const ENe
 		InOutNetwork->Run();
 		CPUOutputs.Emplace(InOutNetwork->GetOutputTensor().GetArrayCopy<float>());
 	}
-	// Input CPU + Network GPU + Output CPU
-	for (int32 Index = 0; Index < InputArrays.Num(); ++Index)
-	{
-		InOutNetwork->SetInputFromArrayCopy(InputArrays[Index]);
-		InOutNetwork->SetDeviceType(/*DeviceType*/ENeuralDeviceType::GPU, /*InputDeviceType*/ENeuralDeviceType::CPU, /*OutputDeviceType*/ENeuralDeviceType::CPU);
-		InOutNetwork->Run();
-		CPUGPUCPUOutputs.Emplace(InOutNetwork->GetOutputTensor().GetArrayCopy<float>());
-	}
-	// Input CPU + Network GPU + Output GPU
-	for (int32 Index = 0; Index < InputArrays.Num(); ++Index)
+	if (UNeuralNetwork::IsGPUConfigCompatibleForUEAndORTBackEnd())
 	{
+		// Input CPU + Network GPU + Output CPU
+		for (int32 Index = 0; Index < InputArrays.Num(); ++Index)
+		{
+			InOutNetwork->SetInputFromArrayCopy(InputArrays[Index]);
+			InOutNetwork->SetDeviceType(/*DeviceType*/ENeuralDeviceType::GPU, /*InputDeviceType*/ENeuralDeviceType::CPU, /*OutputDeviceType*/ENeuralDeviceType::CPU);
+			InOutNetwork->Run();
+			CPUGPUCPUOutputs.Emplace(InOutNetwork->GetOutputTensor().GetArrayCopy<float>());
+		}
+		// Input CPU + Network GPU + Output GPU
+		for (int32 Index = 0; Index < InputArrays.Num(); ++Index)
+		{
 // @todo: Remove if/else (and else contents) once DX12 version of NNI (UEAndORT) is ready
 if (InBackEnd == ENeuralBackEnd::UEOnly /*|| true*/)
 {
-		InOutNetwork->SetInputFromArrayCopy(InputArrays[Index]);
-		InOutNetwork->SetDeviceType(/*DeviceType*/ENeuralDeviceType::GPU, /*InputDeviceType*/ENeuralDeviceType::CPU, /*OutputDeviceType*/ENeuralDeviceType::GPU);
-		InOutNetwork->Run();
-		InOutNetwork->OutputTensorsToCPU();
-		CPUGPUGPUOutputs.Emplace(InOutNetwork->GetOutputTensor().GetArrayCopy<float>());
+			InOutNetwork->SetInputFromArrayCopy(InputArrays[Index]);
+			InOutNetwork->SetDeviceType(/*DeviceType*/ENeuralDeviceType::GPU, /*InputDeviceType*/ENeuralDeviceType::CPU, /*OutputDeviceType*/ENeuralDeviceType::GPU);
+			InOutNetwork->Run();
+			InOutNetwork->OutputTensorsToCPU();
+			CPUGPUGPUOutputs.Emplace(InOutNetwork->GetOutputTensor().GetArrayCopy<float>());
 }
 else
 {
-	UE_LOG(LogNeuralNetworkInferenceQA, Display, TEXT("FModelUnitTester::ModelAccuracyTest(): GPU output for UEAndORT back end not working yet. Uncomment this line to test it."));
-	CPUGPUGPUOutputs.Push(CPUGPUCPUOutputs[Index]);
+		UE_LOG(LogNeuralNetworkInferenceQA, Display, TEXT("FModelUnitTester::ModelAccuracyTest(): GPU output for UEAndORT back end not working yet. Uncomment this line to test it."));
+		CPUGPUGPUOutputs.Push(CPUGPUCPUOutputs[Index]);
 }
-	}
-	// Input GPU + Network GPU + Output CPU
-	for (int32 Index = 0; Index < InputArrays.Num(); ++Index)
-	{
+		}
+		// Input GPU + Network GPU + Output CPU
+		for (int32 Index = 0; Index < InputArrays.Num(); ++Index)
+		{
 // @todo: Remove if/else (and else contents) once DX12 version of NNI (UEAndORT) is ready
 if (InBackEnd == ENeuralBackEnd::UEOnly /*|| true*/)
 {
-		InOutNetwork->SetInputFromArrayCopy(InputArrays[Index]);
-		InOutNetwork->InputTensorsToGPU();
-		InOutNetwork->SetDeviceType(/*DeviceType*/ENeuralDeviceType::GPU, /*InputDeviceType*/ENeuralDeviceType::GPU, /*OutputDeviceType*/ENeuralDeviceType::CPU);
-		InOutNetwork->Run();
-		CPUGPUGPUOutputs.Emplace(InOutNetwork->GetOutputTensor().GetArrayCopy<float>());
+			InOutNetwork->SetInputFromArrayCopy(InputArrays[Index]);
+			InOutNetwork->InputTensorsToGPU();
+			InOutNetwork->SetDeviceType(/*DeviceType*/ENeuralDeviceType::GPU, /*InputDeviceType*/ENeuralDeviceType::GPU, /*OutputDeviceType*/ENeuralDeviceType::CPU);
+			InOutNetwork->Run();
+			GPUGPUCPUOutputs.Emplace(InOutNetwork->GetOutputTensor().GetArrayCopy<float>());
 }
 else
 {
-	UE_LOG(LogNeuralNetworkInferenceQA, Display, TEXT("FModelUnitTester::ModelAccuracyTest(): GPU output for UEAndORT back end not working yet. Uncomment this line to test it."));
-	CPUGPUGPUOutputs.Push(CPUGPUCPUOutputs[Index]);
+		UE_LOG(LogNeuralNetworkInferenceQA, Display, TEXT("FModelUnitTester::ModelAccuracyTest(): GPU output for UEAndORT back end not working yet. Uncomment this line to test it."));
+		GPUGPUCPUOutputs.Push(CPUGPUCPUOutputs[Index]);
 }
+		}
+	}
+	else
+	{
+		UE_LOG(LogNeuralNetworkInferenceQA, Display, TEXT("-------------------- ModelAccuracyTest-UEAndORT-GPU skipped (DX12 not enabled)."));
 	}
 	for (int32 Index = 0; Index < InputArrays.Num(); ++Index)
 	{
+		const TArray<float>& GPUGPUCPUOutput = (UNeuralNetwork::IsGPUConfigCompatibleForUEAndORTBackEnd() ? GPUGPUCPUOutputs[Index] : CPUOutputs[Index]);
+		const TArray<float>& CPUGPUCPUOutput = (UNeuralNetwork::IsGPUConfigCompatibleForUEAndORTBackEnd() ? CPUGPUCPUOutputs[Index] : CPUOutputs[Index]);
+		const TArray<float>& CPUGPUGPUOutput = (UNeuralNetwork::IsGPUConfigCompatibleForUEAndORTBackEnd() ? CPUGPUGPUOutputs[Index] : CPUOutputs[Index]);
 		// Prepare verbose
 		const double CPUAvgL1Norm = GetAveragedL1Norm(CPUOutputs[Index]);
-		const double CPUGPUCPUAvgL1Norm = GetAveragedL1Norm(CPUGPUCPUOutputs[Index]);
-		const double CPUGPUGPUAvgL1Norm = GetAveragedL1Norm(CPUGPUGPUOutputs[Index]);
+		const double CPUGPUCPUAvgL1Norm = GetAveragedL1Norm(CPUGPUCPUOutput);
+		const double GPUGPUCPUAvgL1Norm = GetAveragedL1Norm(GPUGPUCPUOutput);
+		const double CPUGPUGPUAvgL1Norm = GetAveragedL1Norm(CPUGPUGPUOutput);
 		const double RelativeCoefficient = 1. / FMath::Max(1., FMath::Min(CPUAvgL1Norm, CPUGPUCPUAvgL1Norm)); // Max(1, X) to avoid 0s
-		const double CPUGPUAvgL1NormDiff = GetAveragedL1NormDiff(CPUOutputs[Index], CPUGPUCPUOutputs[Index]) * RelativeCoefficient * 1e6;
-		const double GPUGPUAvgL1NormDiff = GetAveragedL1NormDiff(CPUGPUCPUOutputs[Index], CPUGPUGPUOutputs[Index]) / FMath::Max(1., FMath::Min(CPUGPUGPUAvgL1Norm, CPUGPUCPUAvgL1Norm)) * 1e12;
+		const double CPUGPUAvgL1NormDiff = GetAveragedL1NormDiff(CPUOutputs[Index], CPUGPUCPUOutput) * RelativeCoefficient * 1e6;
+		const double GPUGPUInputAvgL1NormDiff = GetAveragedL1NormDiff(CPUGPUCPUOutput, GPUGPUCPUOutput) / FMath::Max(1., FMath::Min(GPUGPUCPUAvgL1Norm, CPUGPUCPUAvgL1Norm)) * 1e12;
+		const double GPUGPUOutputAvgL1NormDiff = GetAveragedL1NormDiff(CPUGPUCPUOutput, CPUGPUGPUOutput) / FMath::Max(1., FMath::Min(CPUGPUGPUAvgL1Norm, CPUGPUCPUAvgL1Norm)) * 1e12;
 		const double FastCPUGPUAvgL1NormDiff = FMath::Abs((CPUAvgL1Norm - CPUGPUCPUAvgL1Norm)) * RelativeCoefficient * 1e6;
 		const double FastCPUAvgL1NormDiff = FMath::Abs(CPUAvgL1Norm - InCPUGroundTruths[Index]) / FMath::Max(1., FMath::Min(CPUAvgL1Norm, InCPUGroundTruths[Index])) * 1e7;
 		const double FastGPUAvgL1NormDiff = FMath::Abs(CPUGPUCPUAvgL1Norm - InGPUGroundTruths[Index]) / FMath::Max(1., FMath::Min(CPUGPUCPUAvgL1Norm, InGPUGroundTruths[Index])) * 1e7;
 		// Print verbose
 		UE_LOG(LogNeuralNetworkInferenceQA, Display, TEXT("%s: InputNorm = %f, OutputNormCPU = %f, OutputNormGPU = %f, OutputNormCPUGPUGPU = %f, OutputNormGT = %f, CPUAvgL1Norm = %f, CPUGPUCPUAvgL1Norm = %f,"),
 			*BackEndString, GetAveragedL1Norm(InputArrays[Index]), CPUAvgL1Norm, CPUGPUCPUAvgL1Norm, CPUGPUGPUAvgL1Norm, InCPUGroundTruths[Index], CPUAvgL1Norm, CPUGPUCPUAvgL1Norm);
-		UE_LOG(LogNeuralNetworkInferenceQA, Display, TEXT("\tCPUGPUAvgL1NormDiff = %fe-6, GPUGPUAvgL1NormDiff = %fe-7, FastCPUGPUAvgL1NormDiff = %fe-6, FastCPUAvgL1NormDiff = %fe-7, FastGPUAvgL1NormDiff = %fe-7 (1e-7 is roughly the precision for float)."),
-			CPUGPUAvgL1NormDiff, GPUGPUAvgL1NormDiff, FastCPUGPUAvgL1NormDiff, FastCPUAvgL1NormDiff, FastGPUAvgL1NormDiff);
+		UE_LOG(LogNeuralNetworkInferenceQA, Display, TEXT("\tCPUGPUAvgL1NormDiff = %fe-6, GPUGPUInputAvgL1NormDiff = %fe-7, GPUGPUOutputAvgL1NormDiff = %fe-7, FastCPUGPUAvgL1NormDiff = %fe-6, FastCPUAvgL1NormDiff = %fe-7, FastGPUAvgL1NormDiff = %fe-7 (1e-7 is roughly the precision for float)."),
+			CPUGPUAvgL1NormDiff, GPUGPUInputAvgL1NormDiff, GPUGPUOutputAvgL1NormDiff, FastCPUGPUAvgL1NormDiff, FastCPUAvgL1NormDiff, FastGPUAvgL1NormDiff);
 		// Check if test failed and (if so) display information
 		const bool bDidSomeTestFailed = (!FMath::IsFinite(FastCPUGPUAvgL1NormDiff) || FastCPUGPUAvgL1NormDiff > 5)
 			|| (!FMath::IsFinite(CPUGPUAvgL1NormDiff) || CPUGPUAvgL1NormDiff > 5)
-			|| (!FMath::IsFinite(GPUGPUAvgL1NormDiff) || GPUGPUAvgL1NormDiff > 1)
+			|| (!FMath::IsFinite(GPUGPUInputAvgL1NormDiff) || GPUGPUInputAvgL1NormDiff > 1)
+			|| (!FMath::IsFinite(GPUGPUOutputAvgL1NormDiff) || GPUGPUOutputAvgL1NormDiff > 1)
 			|| (!FMath::IsFinite(FastCPUAvgL1NormDiff) || FastCPUAvgL1NormDiff > 30)
 			|| (!FMath::IsFinite(FastGPUAvgL1NormDiff) || FastGPUAvgL1NormDiff > 30);
 		if (bDidSomeTestFailed)
@@ -305,15 +335,17 @@ else
 			const TArray<int64>& OutputSizes = InOutNetwork->GetOutputTensor().GetSizes();
 			UE_LOG(LogNeuralNetworkInferenceQA, Display, TEXT("FastCPUGPUAvgL1NormDiff (%fe-6) < 5e-6 might have failed."), FastCPUGPUAvgL1NormDiff);
 			UE_LOG(LogNeuralNetworkInferenceQA, Display, TEXT("CPUGPUAvgL1NormDiff (%fe-6) < 5e-6 might have failed."), CPUGPUAvgL1NormDiff);
-			UE_LOG(LogNeuralNetworkInferenceQA, Display, TEXT("GPUGPUAvgL1NormDiff (%fe-12) < 1e-12 might have failed."), GPUGPUAvgL1NormDiff);
+			UE_LOG(LogNeuralNetworkInferenceQA, Display, TEXT("GPUGPUInputAvgL1NormDiff (%fe-12) < 1e-12 might have failed."), GPUGPUInputAvgL1NormDiff);
+			UE_LOG(LogNeuralNetworkInferenceQA, Display, TEXT("GPUGPUOutputAvgL1NormDiff (%fe-12) < 1e-12 might have failed."), GPUGPUOutputAvgL1NormDiff);
 			UE_LOG(LogNeuralNetworkInferenceQA, Display, TEXT("FastCPUAvgL1NormDiff (%fe-7) < 30e-7 might have failed (~30 times the float precision).\nCPUOutput = %s."), FastCPUAvgL1NormDiff, *FNeuralTensor(CPUOutputs[Index], OutputSizes).ToString(MaxNumberElementsToDisplay));
-			UE_LOG(LogNeuralNetworkInferenceQA, Display, TEXT("FastGPUAvgL1NormDiff (%fe-7) < 30e-7 might have failed (~30 times the float precision).\nCPUGPUCPUOutput = %s."), FastGPUAvgL1NormDiff, *FNeuralTensor(CPUGPUCPUOutputs[Index], OutputSizes).ToString(MaxNumberElementsToDisplay));
+			UE_LOG(LogNeuralNetworkInferenceQA, Display, TEXT("FastGPUAvgL1NormDiff (%fe-7) < 30e-7 might have failed (~30 times the float precision).\nCPUGPUCPUOutput = %s."), FastGPUAvgL1NormDiff, *FNeuralTensor(CPUGPUCPUOutput, OutputSizes).ToString(MaxNumberElementsToDisplay));
 			UE_LOG(LogNeuralNetworkInferenceQA, Display, TEXT("Input = %s"),
 				*FNeuralTensor(InOutNetwork->GetInputTensor().GetArrayCopy<float>(), InputSizes).ToString(MaxNumberElementsToDisplay));
 			UE_LOG(LogNeuralNetworkInferenceQA, Display, TEXT("CPUOutput = %s"), *FNeuralTensor(CPUOutputs[Index], OutputSizes).ToString(MaxNumberElementsToDisplay));
-			UE_LOG(LogNeuralNetworkInferenceQA, Display, TEXT("CPUGPUCPUOutput = %s"), *FNeuralTensor(CPUGPUCPUOutputs[Index], OutputSizes).ToString(MaxNumberElementsToDisplay));
-			UE_LOG(LogNeuralNetworkInferenceQA, Display, TEXT("CPUGPUGPUOutput = %s"), *FNeuralTensor(CPUGPUGPUOutputs[Index], OutputSizes).ToString(MaxNumberElementsToDisplay));
-			UE_LOG(LogNeuralNetworkInferenceQA, Warning, TEXT("At least 1 of the 5 CPU/GPU tests failed."));
+			UE_LOG(LogNeuralNetworkInferenceQA, Display, TEXT("CPUGPUCPUOutput = %s"), *FNeuralTensor(CPUGPUCPUOutput, OutputSizes).ToString(MaxNumberElementsToDisplay));
+			UE_LOG(LogNeuralNetworkInferenceQA, Display, TEXT("GPUGPUCPUOutput = %s"), *FNeuralTensor(GPUGPUCPUOutput, OutputSizes).ToString(MaxNumberElementsToDisplay));
+			UE_LOG(LogNeuralNetworkInferenceQA, Display, TEXT("CPUGPUGPUOutput = %s"), *FNeuralTensor(CPUGPUGPUOutput, OutputSizes).ToString(MaxNumberElementsToDisplay));
+			UE_LOG(LogNeuralNetworkInferenceQA, Warning, TEXT("At least 1 of the 6 CPU/GPU tests failed."));
 			return false;
 		}
 		bDidGlobalTestPassed &= !bDidSomeTestFailed;
