NeuralNetworkInference (NNI) plugin:
- Enable DirectML debug layer

#rnx

diff --git a/Engine/Plugins/Experimental/NeuralNetworkInference/Source/NeuralNetworkInference/Private/NeuralNetworkImplBackEndUEAndORT.cpp b/Engine/Plugins/Experimental/NeuralNetworkInference/Source/NeuralNetworkInference/Private/NeuralNetworkImplBackEndUEAndORT.cpp
--- a/Engine/Plugins/Experimental/NeuralNetworkInference/Source/NeuralNetworkInference/Private/NeuralNetworkImplBackEndUEAndORT.cpp
+++ b/Engine/Plugins/Experimental/NeuralNetworkInference/Source/NeuralNetworkInference/Private/NeuralNetworkImplBackEndUEAndORT.cpp
@@ -100,14 +100,27 @@ IDMLDevice* FPrivateImplBackEndUEAndORT::FDMLDeviceList::Add(ID3D12Device* Devic
 
 	DML_CREATE_DEVICE_FLAGS DmlCreateFlags = DML_CREATE_DEVICE_FLAG_NONE;
 
-	DMLCreateDevice1(Device,
-		DmlCreateFlags,
-		DML_FEATURE_LEVEL_2_0,
-		IID_PPV_ARGS(&DmlDevice));
+#if !UE_BUILD_SHIPPING
+	if (D3D12RHI_ShouldCreateWithD3DDebug()
+		|| FParse::Param(FCommandLine::Get(), TEXT("d3d12gpuvalidation")) || FParse::Param(FCommandLine::Get(), TEXT("gpuvalidation")))
+	{
+		DmlCreateFlags |= DML_CREATE_DEVICE_FLAG_DEBUG;
+	}
+#endif
+
+	HRESULT res = DMLCreateDevice1(Device, DmlCreateFlags, DML_FEATURE_LEVEL_2_0, IID_PPV_ARGS(&DmlDevice));
+
+	// Handle the case if Graphics Debug Tools are not installed
+	if (res == DXGI_ERROR_SDK_COMPONENT_MISSING)
+	{
+		DmlCreateFlags &= ~DML_CREATE_DEVICE_FLAG_DEBUG;
+
+		res = DMLCreateDevice1(Device, DmlCreateFlags, DML_FEATURE_LEVEL_2_0, IID_PPV_ARGS(&DmlDevice));
+	}
 
 	if (!DmlDevice)
 	{
-		UE_LOG(LogNeuralNetworkInference, Warning, TEXT("FDMLDeviceList::Add(): Failed to create DML device."));
+		UE_LOG(LogNeuralNetworkInference, Warning, TEXT("FDMLDeviceList::Add(): Failed to create DML device, res=%0x."), res);
 		return nullptr;
 	}
 
