fixed uninitialized struct members
#rnx

diff --git a/Engine/Plugins/Experimental/NeuralNetworkInference/Source/Internal/ModelProto/Internal/ModelProto.h b/Engine/Plugins/Experimental/NeuralNetworkInference/Source/Internal/ModelProto/Internal/ModelProto.h
--- a/Engine/Plugins/Experimental/NeuralNetworkInference/Source/Internal/ModelProto/Internal/ModelProto.h
+++ b/Engine/Plugins/Experimental/NeuralNetworkInference/Source/Internal/ModelProto/Internal/ModelProto.h
@@ -316,7 +316,7 @@ public:
 	 * This field MUST refer to an integral type ([U]INT{8|16|32|64}) or STRING
 	 */
 	UPROPERTY(VisibleAnywhere, Category = "Neural Network Inference")
-	int32 KeyType;
+	int32 KeyType =  0;
 
 	/**
 	 * This field MUST be present for this version of the IR.
diff --git a/Engine/Plugins/VirtualProduction/CameraCalibrationCore/Source/CameraCalibrationCore/Public/LensData.h b/Engine/Plugins/VirtualProduction/CameraCalibrationCore/Source/CameraCalibrationCore/Public/LensData.h
--- a/Engine/Plugins/VirtualProduction/CameraCalibrationCore/Source/CameraCalibrationCore/Public/LensData.h
+++ b/Engine/Plugins/VirtualProduction/CameraCalibrationCore/Source/CameraCalibrationCore/Public/LensData.h
@@ -50,7 +50,7 @@ public:
 
 	/** Camera image size, in pixels */
 	UPROPERTY(EditAnywhere, BlueprintReadWrite, Category = "Lens Info")
-	FIntPoint ImageDimensions;
+	FIntPoint ImageDimensions = FIntPoint(0, 1);
 
 };
 
