NeuralNetworkInference plugin:
Rename/Move the QA Timer Class, moving it to NeuralNetworkInference

diff --git a/Engine/Plugins/Experimental/NeuralNetworkInference/Source/NeuralNetworkInference/Internal/NeuralNetworkInferenceTimer.cpp b/Engine/Plugins/Experimental/NeuralNetworkInference/Source/NeuralNetworkInference/Internal/NeuralNetworkInferenceTimer.cpp
new file mode 100644
--- /dev/null
+++ b/Engine/Plugins/Experimental/NeuralNetworkInference/Source/NeuralNetworkInference/Internal/NeuralNetworkInferenceTimer.cpp
@@ -0,0 +1,17 @@
+// Copyright Epic Games, Inc. All Rights Reserved.
+
+#include "NeuralNetworkInferenceTimer.h"
+#include "HAL/PlatformTime.h"
+
+/* FNeuralNetworkInferenceQATimer public functions
+ *****************************************************************************/
+
+void FNeuralNetworkInferenceTimer::Tic()
+{
+	TimeStart = FPlatformTime::Seconds();
+}
+
+float FNeuralNetworkInferenceTimer::Toc() const
+{
+	return (FPlatformTime::Seconds() - TimeStart) * 1e3;
+}
diff --git a/Engine/Plugins/Experimental/NeuralNetworkInference/Source/NeuralNetworkInference/Public/NeuralNetworkInferenceTimer.h b/Engine/Plugins/Experimental/NeuralNetworkInference/Source/NeuralNetworkInference/Public/NeuralNetworkInferenceTimer.h
new file mode 100644
--- /dev/null
+++ b/Engine/Plugins/Experimental/NeuralNetworkInference/Source/NeuralNetworkInference/Public/NeuralNetworkInferenceTimer.h
@@ -0,0 +1,19 @@
+// Copyright Epic Games, Inc. All Rights Reserved.
+
+#pragma once
+
+#include "CoreMinimal.h"
+
+class NEURALNETWORKINFERENCE_API FNeuralNetworkInferenceTimer
+{
+public:
+	void Tic();
+
+	/**
+	 * Time in milliseconds, but with nanosecond accuracy (e.g., 1.234567 msec).
+	 */
+	float Toc() const;
+
+private:
+	double TimeStart;
+};
diff --git a/Engine/Plugins/Experimental/NeuralNetworkInference/Source/NeuralNetworkInferenceQA/Private/LegacyModelUnitTester.cpp b/Engine/Plugins/Experimental/NeuralNetworkInference/Source/NeuralNetworkInferenceQA/Private/LegacyModelUnitTester.cpp
--- a/Engine/Plugins/Experimental/NeuralNetworkInference/Source/NeuralNetworkInferenceQA/Private/LegacyModelUnitTester.cpp
+++ b/Engine/Plugins/Experimental/NeuralNetworkInference/Source/NeuralNetworkInferenceQA/Private/LegacyModelUnitTester.cpp
@@ -2,7 +2,7 @@
 
 #include "LegacyModelUnitTester.h"
 #include "ModelProtoFileReaderTester.h"
-#include "NeuralNetworkInferenceQATimer.h"
+#include "NeuralNetworkInferenceTimer.h"
 #include "NeuralNetworkInferenceQAUtils.h"
 #include "Misc/Paths.h"
 
@@ -262,7 +262,7 @@ bool FLegacyModelUnitTester::ModelSpeedTest(const FString& InUAssetPath, const i
 	InputArray.Init(1.f, NetworkSize);
 	TArray<float> OutputArrayGPU;
 	// Speed profiling
-	FNeuralNetworkInferenceQATimer Timer;
+	FNeuralNetworkInferenceTimer Timer;
 	// CPU
 	InOutNetwork->SetDeviceType(ENeuralDeviceType::CPU);
 	LEGACY_MODEL_UNIT_TESTER_SPEED_TEST(CPUCopyTimer1, CPUNetworkTimer1, DummyCPU1, 1, /*bIsGPU*/false);
diff --git a/Engine/Plugins/Experimental/NeuralNetworkInference/Source/NeuralNetworkInferenceQA/Private/ModelUnitTester.cpp b/Engine/Plugins/Experimental/NeuralNetworkInference/Source/NeuralNetworkInferenceQA/Private/ModelUnitTester.cpp
--- a/Engine/Plugins/Experimental/NeuralNetworkInference/Source/NeuralNetworkInferenceQA/Private/ModelUnitTester.cpp
+++ b/Engine/Plugins/Experimental/NeuralNetworkInference/Source/NeuralNetworkInferenceQA/Private/ModelUnitTester.cpp
@@ -1,7 +1,7 @@
 // Copyright Epic Games, Inc. All Rights Reserved.
 
 #include "ModelUnitTester.h"
-#include "NeuralNetworkInferenceQATimer.h"
+#include "NeuralNetworkInferenceTimer.h"
 #include "NeuralNetworkInferenceQAUtils.h"
 #include "Misc/Paths.h"
 
@@ -416,7 +416,7 @@ bool FModelUnitTester::ModelSpeedTest(const FString& InUAssetPath, const ENeural
 	InputArray.Init(1.f, NetworkSize);
 	TArray<float> CPUGPUCPUOutput;
 	// Speed profiling
-	FNeuralNetworkInferenceQATimer Timer;
+	FNeuralNetworkInferenceTimer Timer;
 	// Run profiling 1 time
 	InOutNetwork->SetDeviceType(InDeviceType);
 	MODEL_UNIT_TESTER_SPEED_TEST(CopyTimer1, NetworkTimer1, 1, /*bIsGPU*/false);
diff --git a/Engine/Plugins/Experimental/NeuralNetworkInference/Source/NeuralNetworkInferenceQA/Private/NeuralNetworkInferenceQATimer.cpp b/Engine/Plugins/Experimental/NeuralNetworkInference/Source/NeuralNetworkInferenceQA/Private/NeuralNetworkInferenceQATimer.cpp
deleted file mode 100644
diff --git a/Engine/Plugins/Experimental/NeuralNetworkInference/Source/NeuralNetworkInferenceQA/Public/NeuralNetworkInferenceQATimer.h b/Engine/Plugins/Experimental/NeuralNetworkInference/Source/NeuralNetworkInferenceQA/Public/NeuralNetworkInferenceQATimer.h
deleted file mode 100644
