NeuralNetworkInference (NNI) plugin: In order to remove NNI-Legacy, we need to extend UNeuralNetwork with:
- Functions to copy input and output tensors.
- Function to load from FNeuralTensor and FNeuralOperators.

#rnx

diff --git a/Engine/Plugins/Experimental/NeuralNetworkInference/Source/NeuralNetworkInference/Private/NeuralNetwork.cpp b/Engine/Plugins/Experimental/NeuralNetworkInference/Source/NeuralNetworkInference/Private/NeuralNetwork.cpp
--- a/Engine/Plugins/Experimental/NeuralNetworkInference/Source/NeuralNetworkInference/Private/NeuralNetwork.cpp
+++ b/Engine/Plugins/Experimental/NeuralNetworkInference/Source/NeuralNetworkInference/Private/NeuralNetwork.cpp
@@ -228,7 +228,7 @@ bool UNeuralNetwork::IsGPUConfigCompatibleForCurrentBackEnd() const
 	/* UEOnly */ \
 	else if (BackEndForCurrentPlatform == ENeuralBackEnd::UEOnly) \
 	{ \
-		FNeuralTensorManager& TensorManager = ImplBackEndUEOnly->TensorManager; \
+		FNeuralTensorManager& TensorManager = *ImplBackEndUEOnly->TensorManager; \
 		return TensorManager.UEOnly_GetTensorsFunction()[TensorManager.UEAndORTOnly_GetIndexes()[InTensorIndex]]; \
 	} \
 	/* Unknown */ \
@@ -257,7 +257,7 @@ void UNeuralNetwork::SetInputFromArrayCopy(const TArray<float>& InArray, const i
 	// UEOnly
 	else if (BackEndForCurrentPlatform == ENeuralBackEnd::UEOnly)
 	{
-		FNeuralTensorManager& TensorManager = ImplBackEndUEOnly->TensorManager;
+		FNeuralTensorManager& TensorManager = *ImplBackEndUEOnly->TensorManager;
 		TensorManager.GetTensorsMutable()[TensorManager.GetInputIndexes()[InTensorIndex]].SetFromArrayCopy(InArray);
 	}
 
@@ -286,7 +286,7 @@ void* UNeuralNetwork::GetInputDataPointerMutable(const int32 InTensorIndex)
 	// UEOnly
 	else if (BackEndForCurrentPlatform == ENeuralBackEnd::UEOnly)
 	{
-		FNeuralTensorManager& TensorManager = ImplBackEndUEOnly->TensorManager;
+		FNeuralTensorManager& TensorManager = *ImplBackEndUEOnly->TensorManager;
 		return TensorManager.GetTensorsMutable()[TensorManager.GetInputIndexes()[InTensorIndex]].GetData();
 	}
 
@@ -313,7 +313,7 @@ int64 UNeuralNetwork::GetInputTensorNumber() const
 	// UEOnly
 	else if (BackEndForCurrentPlatform == ENeuralBackEnd::UEOnly)
 	{
-		return ImplBackEndUEOnly->TensorManager.GetInputIndexes().Num();
+		return ImplBackEndUEOnly->TensorManager->GetInputIndexes().Num();
 	}
 
 	// Unknown
@@ -344,7 +344,7 @@ int64 UNeuralNetwork::GetOutputTensorNumber() const
 	// UEOnly
 	else if (BackEndForCurrentPlatform == ENeuralBackEnd::UEOnly)
 	{
-		return ImplBackEndUEOnly->TensorManager.GetOutputIndexes().Num();
+		return ImplBackEndUEOnly->TensorManager->GetOutputIndexes().Num();
 	}
 
 	// Unknown
@@ -352,6 +352,42 @@ int64 UNeuralNetwork::GetOutputTensorNumber() const
 	return -1;
 }
 
+TArray<FNeuralTensor> UNeuralNetwork::CreateInputArrayCopy() const
+{
+	TArray<FNeuralTensor> InputTensorArray;
+	for (uint32 InputTensorIndex = 0; InputTensorIndex < GetInputTensorNumber(); ++InputTensorIndex)
+	{
+		const FNeuralTensor& InputTensor = GetInputTensor(InputTensorIndex);
+		InputTensorArray.Push(InputTensor);
+	}
+	return InputTensorArray;
+}
+
+void UNeuralNetwork::SetInputFromArrayCopy(const TArray<FNeuralTensor>& InInputTensorArray)
+{
+	if (GetInputTensorNumber() != InInputTensorArray.Num())
+	{
+		UE_LOG(LogNeuralNetworkInference, Warning, TEXT("UNeuralNetwork::SetInputFromArrayCopy(): GetInputTensorNumber() == InInputTensorArray.Num() failed, %d != %d."), GetInputTensorNumber(), InInputTensorArray.Num());
+		return;
+	}
+	for (uint32 InputTensorIndex = 0; InputTensorIndex < GetInputTensorNumber(); ++InputTensorIndex)
+	{
+		FNeuralTensor& InputTensor = GetInputTensorMutable(InputTensorIndex);
+		InputTensor.SetFromUnderlyingUInt8ArrayCopy(InInputTensorArray[InputTensorIndex].GetUnderlyingUInt8ArrayRef());
+	}
+}
+
+TArray<FNeuralTensor> UNeuralNetwork::CreateOutputArrayCopy() const
+{
+	TArray<FNeuralTensor> OutputTensorArray;
+	for (uint32 OutputTensorIndex = 0; OutputTensorIndex < GetOutputTensorNumber(); ++OutputTensorIndex)
+	{
+		const FNeuralTensor& OutputTensor = GetOutputTensor(OutputTensorIndex);
+		OutputTensorArray.Push(OutputTensor);
+	}
+	return OutputTensorArray;
+}
+
 void UNeuralNetwork::InputTensorsToGPU(const TArray<int32>& InTensorIndexes)
 {
 	// Sanity check
@@ -508,6 +544,15 @@ FNeuralTensor& UNeuralNetwork::GetOutputTensorMutable(const int32 InTensorIndex)
 	GET_TENSOR_CODE(InTensorIndex, GetOutputIndexes, GetTensorsMutable, OutputTensors);
 }
 
+bool UNeuralNetwork::Load(TSharedPtr<FNeuralTensorManager>& InTensorManager, const TArray<TSharedPtr<FNeuralOperator>>& InOperators)
+{
+	DECLARE_SCOPE_CYCLE_COUNTER(TEXT("UNeuralNetwork_Load_FromTensorManagerAndOperators"), STAT_UNeuralNetwork_Load, STATGROUP_MachineLearning);
+	BackEnd = ENeuralBackEnd::UEOnly;
+	BackEndForCurrentPlatform = ENeuralBackEnd::UEOnly;
+	bIsLoaded = UNeuralNetwork::FImplBackEndUEOnly::Load(ImplBackEndUEOnly, InTensorManager, InOperators);
+	return bIsLoaded;
+}
+
 #if WITH_EDITOR
 UAssetImportData* UNeuralNetwork::GetAssetImportData() const
 {
diff --git a/Engine/Plugins/Experimental/NeuralNetworkInference/Source/NeuralNetworkInference/Private/NeuralNetworkImplBackEndUEOnly.cpp b/Engine/Plugins/Experimental/NeuralNetworkInference/Source/NeuralNetworkInference/Private/NeuralNetworkImplBackEndUEOnly.cpp
--- a/Engine/Plugins/Experimental/NeuralNetworkInference/Source/NeuralNetworkInference/Private/NeuralNetworkImplBackEndUEOnly.cpp
+++ b/Engine/Plugins/Experimental/NeuralNetworkInference/Source/NeuralNetworkInference/Private/NeuralNetworkImplBackEndUEOnly.cpp
@@ -15,18 +15,8 @@
 
 bool UNeuralNetwork::FImplBackEndUEOnly::Load(TSharedPtr<FImplBackEndUEOnly>& InOutImplBackEndUEOnly, const TArray<uint8>& InModelReadFromFileInBytes)
 {
-	// Initialize InOutImplBackEndUEOnly
-	if (!InOutImplBackEndUEOnly.IsValid())
-	{
-		InOutImplBackEndUEOnly = MakeShared<FImplBackEndUEOnly>();
-	}
-	// Clean previous networks
-	else
-	{
-		InOutImplBackEndUEOnly->Operators.Empty();
-		InOutImplBackEndUEOnly->ModelProto = FModelProto();
-		InOutImplBackEndUEOnly->bAreTensorsInGpu = false;
-	}
+	// Reset InOutImplBackEndUEOnly
+	Reset(InOutImplBackEndUEOnly);
 	// Read ModelProto
 	if (!FModelProtoFileReader::ReadModelProtoFromArray(InOutImplBackEndUEOnly->ModelProto, InModelReadFromFileInBytes))
 	{
@@ -39,20 +29,39 @@ bool UNeuralNetwork::FImplBackEndUEOnly::Load(TSharedPtr<FImplBackEndUEOnly>& In
 		return false;
 	}
 	// Turn ModelProto into Operators
-	if (!FGraphProtoToNeuralNetworkConverter::Translate(InOutImplBackEndUEOnly->Operators, InOutImplBackEndUEOnly->TensorManager, InOutImplBackEndUEOnly->ModelProto.GetGraph(), /*bInIsTensorManagerConst*/false))
+	if (!FGraphProtoToNeuralNetworkConverter::Translate(InOutImplBackEndUEOnly->Operators, *InOutImplBackEndUEOnly->TensorManager, InOutImplBackEndUEOnly->ModelProto.GetGraph(), /*bInIsTensorManagerConst*/false))
 	{
 		UE_LOG(LogNeuralNetworkInference, Warning, TEXT("FImplBackEndUEOnly::Load(): Translate failed."));
 		return false;
 	}
-	if (!InOutImplBackEndUEOnly->TensorManager.IsLoaded())
+	if (!InOutImplBackEndUEOnly->TensorManager->IsLoaded())
 	{
-		UE_LOG(LogNeuralNetworkInference, Warning, TEXT("FImplBackEndUEOnly::Load(): TensorManager.IsLoaded() returned false."));
+		UE_LOG(LogNeuralNetworkInference, Warning, TEXT("FImplBackEndUEOnly::Load(): TensorManager->IsLoaded() returned false."));
 		return false;
 	}
 	// Load successful
 	return true;
 }
 
+bool UNeuralNetwork::FImplBackEndUEOnly::Load(TSharedPtr<FImplBackEndUEOnly>& InOutImplBackEndUEOnly, TSharedPtr<FNeuralTensorManager>& InTensorManager, const TArray<TSharedPtr<FNeuralOperator>>& InOperators)
+{
+	// Reset InOutImplBackEndUEOnly
+	Reset(InOutImplBackEndUEOnly);
+	// Load
+	if (!InTensorManager.IsValid())
+	{
+		UE_LOG(LogNeuralNetworkInference, Warning, TEXT("UNeuralNetworkLegacy::Load(): TensorManager was a nullptr."));
+		return false;
+	}
+	if (!InTensorManager->IsLoaded())
+	{
+		UE_LOG(LogNeuralNetworkInference, Warning, TEXT("UNeuralNetworkLegacy::Load(): TensorManager could not be loaded."));
+	}
+	Swap(InOutImplBackEndUEOnly->TensorManager, InTensorManager);
+	InOutImplBackEndUEOnly->Operators = InOperators;
+	return (InOutImplBackEndUEOnly->Operators.Num() > 0 && InOutImplBackEndUEOnly->TensorManager->IsLoaded());
+}
+
 //bool UNeuralNetwork::FImplBackEndUEOnly::Load(TSharedPtr<FImplBackEndUEOnly>& InOutImplBackEndUEOnly, FNeuralTensorManager& InTensorManager, const TArray<TSharedPtr<FNeuralOperator>>& InOperators)
 //{
 //	// Initialize InOutImplBackEndUEOnly
@@ -67,7 +76,7 @@ bool UNeuralNetwork::FImplBackEndUEOnly::Load(TSharedPtr<FImplBackEndUEOnly>& In
 //	}
 //	Swap(InOutImplBackEndUEOnly->TensorManager, InTensorManager);
 //	InOutImplBackEndUEOnly->Operators = InOperators;
-//	return (InOutImplBackEndUEOnly->Operators.Num() > 0 && InOutImplBackEndUEOnly->TensorManager.IsLoaded());
+//	return (InOutImplBackEndUEOnly->Operators.Num() > 0 && InOutImplBackEndUEOnly->TensorManager->IsLoaded());
 //}
 
 void UNeuralNetwork::FImplBackEndUEOnly::Run(FOnAsyncRunCompleted& InOutOnAsyncRunCompletedDelegate, const ENeuralNetworkSynchronousMode InSynchronousMode, const ENeuralDeviceType InDeviceType, const ENeuralDeviceType InInputDeviceType, const ENeuralDeviceType InOutputDeviceType)
@@ -94,9 +103,9 @@ void UNeuralNetwork::FImplBackEndUEOnly::Run(FOnAsyncRunCompleted& InOutOnAsyncR
 		{
 			DECLARE_SCOPE_CYCLE_COUNTER(TEXT("UNeuralNetwork_UEOnly_Run::Forward_GPU"), STAT_UNeuralNetwork_UEOnly_Run_Forward_GPU, STATGROUP_MachineLearning);
 			// Sanity check
-			if (TensorManager.GetTensorsMutable().Num() < 1)
+			if (TensorManager->GetTensorsMutable().Num() < 1)
 			{
-				UE_LOG(LogNeuralNetworkInference, Warning, TEXT("FImplBackEndUEOnly::Run(): Tensors.Num() = %d (should be > 0)."), TensorManager.GetTensorsMutable().Num());
+				UE_LOG(LogNeuralNetworkInference, Warning, TEXT("FImplBackEndUEOnly::Run(): Tensors.Num() = %d (should be > 0)."), TensorManager->GetTensorsMutable().Num());
 				return;
 			}
 
@@ -107,7 +116,7 @@ void UNeuralNetwork::FImplBackEndUEOnly::Run(FOnAsyncRunCompleted& InOutOnAsyncR
 					FRDGBuilder GraphBuilder(RHICmdList, RDG_EVENT_NAME("FImplBackEndUEOnly::Run()"));
 
 					// Move memory from CPU to GPU
-					TArray<FNeuralTensor>& Tensors = TensorManager.GetTensorsMutable();
+					TArray<FNeuralTensor>& Tensors = TensorManager->GetTensorsMutable();
 					const bool bIsInputInCPU = (InInputDeviceType == ENeuralDeviceType::CPU);
 					// Move only input tensors to GPU (once per FImplBackEndUEOnly::Run())
 					if (bAreTensorsInGpu)
@@ -115,11 +124,11 @@ void UNeuralNetwork::FImplBackEndUEOnly::Run(FOnAsyncRunCompleted& InOutOnAsyncR
 						// If input is on CPU, move inputs to GPU
 						if (bIsInputInCPU)
 						{
-							for (const int32 InputIndex : TensorManager.GetInputIndexes())
+							for (const int32 InputIndex : TensorManager->GetInputIndexes())
 							{
 								Tensors[InputIndex].ToGPU_RenderThread(&GraphBuilder);
 							}
-							for (const int32 NonInputIndex : TensorManager.GetNonInputIndexes())
+							for (const int32 NonInputIndex : TensorManager->GetNonInputIndexes())
 							{
 								Tensors[NonInputIndex].UpdateSRVAndOrUAV_RenderThread(&GraphBuilder);
 							}
@@ -149,11 +158,11 @@ void UNeuralNetwork::FImplBackEndUEOnly::Run(FOnAsyncRunCompleted& InOutOnAsyncR
 						// If input is already on GPU, just refresh inputs
 						else
 						{
-							for (const int32 InputIndex : TensorManager.GetInputIndexes())
+							for (const int32 InputIndex : TensorManager->GetInputIndexes())
 							{
 								Tensors[InputIndex].UpdateSRVAndOrUAV_RenderThread(&GraphBuilder);
 							}
-							for (const int32 NonInputIndex : TensorManager.GetNonInputIndexes())
+							for (const int32 NonInputIndex : TensorManager->GetNonInputIndexes())
 							{
 								Tensors[NonInputIndex].ToGPU_RenderThread(&GraphBuilder);
 							}
@@ -183,7 +192,7 @@ void UNeuralNetwork::FImplBackEndUEOnly::Run(FOnAsyncRunCompleted& InOutOnAsyncR
 					const bool bIsOutputInCPU = (InOutputDeviceType == ENeuralDeviceType::CPU);
 					if (bIsOutputInCPU)
 					{
-						for (const int32 OutputIndex : TensorManager.GetOutputIndexes())
+						for (const int32 OutputIndex : TensorManager->GetOutputIndexes())
 						{
 							Tensors[OutputIndex].ToCPU_RenderThread(&GraphBuilder);
 						}
@@ -229,8 +238,8 @@ void UNeuralNetwork::FImplBackEndUEOnly::Run(FOnAsyncRunCompleted& InOutOnAsyncR
 //	FString String = ModelProto.GetGraph().ToString();
 //	// Add FNeuralTensor(s)
 //	String += TEXT("TensorManager:\n");
-//	const TMap<FString, int32>& NameIndexMap = TensorManager.GetNameIndexMap();
-//	const TArray<FNeuralTensor>& Tensors = TensorManager.GetTensors();
+//	const TMap<FString, int32>& NameIndexMap = TensorManager->GetNameIndexMap();
+//	const TArray<FNeuralTensor>& Tensors = TensorManager->GetTensors();
 //	if (NameIndexMap.Num() > 0)
 //	{
 //		for (const auto& NameIndexPair : NameIndexMap)
@@ -246,12 +255,12 @@ void UNeuralNetwork::FImplBackEndUEOnly::Run(FOnAsyncRunCompleted& InOutOnAsyncR
 //		}
 //	}
 //	String += TEXT("InputTensorMap:\n");
-//	for (const auto& NameIndexPair : TensorManager.GetInputNameIndexMap())
+//	for (const auto& NameIndexPair : TensorManager->GetInputNameIndexMap())
 //	{
 //		String += FString::Format(TEXT(" -{0}: {1}\n"), { NameIndexPair.Key, Tensors[NameIndexPair.Value].ToString(20) });
 //	}
 //	String += TEXT("OutputTensorMap:\n");
-//	for (const auto& NameIndexPair : TensorManager.GetOutputNameIndexMap())
+//	for (const auto& NameIndexPair : TensorManager->GetOutputNameIndexMap())
 //	{
 //		String += FString::Format(TEXT(" -{0}: {1}\n"), { NameIndexPair.Key, Tensors[NameIndexPair.Value].ToString(20) });
 //	}
@@ -264,3 +273,25 @@ void UNeuralNetwork::FImplBackEndUEOnly::Run(FOnAsyncRunCompleted& InOutOnAsyncR
 //	// Return result
 //	return String;
 //}
+
+
+
+/* FImplBackEndUEOnly private functions
+ *****************************************************************************/
+
+void UNeuralNetwork::FImplBackEndUEOnly::Reset(TSharedPtr<FImplBackEndUEOnly>& InOutImplBackEndUEOnly)
+{
+	// Initialize InOutImplBackEndUEOnly
+	if (!InOutImplBackEndUEOnly.IsValid())
+	{
+		InOutImplBackEndUEOnly = MakeShared<FImplBackEndUEOnly>();
+	}
+	// If already initialized, clean previous network
+	else
+	{
+		InOutImplBackEndUEOnly->Operators.Empty();
+		InOutImplBackEndUEOnly->ModelProto = FModelProto();
+		InOutImplBackEndUEOnly->bAreTensorsInGpu = false;
+	}
+	InOutImplBackEndUEOnly->TensorManager = MakeShared<FNeuralTensorManager>();
+}
diff --git a/Engine/Plugins/Experimental/NeuralNetworkInference/Source/NeuralNetworkInference/Private/NeuralNetworkImplBackEndUEOnly.h b/Engine/Plugins/Experimental/NeuralNetworkInference/Source/NeuralNetworkInference/Private/NeuralNetworkImplBackEndUEOnly.h
--- a/Engine/Plugins/Experimental/NeuralNetworkInference/Source/NeuralNetworkInference/Private/NeuralNetworkImplBackEndUEOnly.h
+++ b/Engine/Plugins/Experimental/NeuralNetworkInference/Source/NeuralNetworkInference/Private/NeuralNetworkImplBackEndUEOnly.h
@@ -23,7 +23,7 @@ struct UNeuralNetwork::FImplBackEndUEOnly
 	/**
 	 * It contains a few TArray and TMaps for all FNeuralTensors (Input, Output, Intermediate(Not)Initialized, Weight).
 	 */
-	FNeuralTensorManager TensorManager;
+	TSharedPtr<FNeuralTensorManager> TensorManager;
 
 	/**
 	 * Only for the vanilla back end.
@@ -31,9 +31,16 @@ struct UNeuralNetwork::FImplBackEndUEOnly
 	 */
 	TArray<TSharedPtr<FNeuralOperator>> Operators;
 
- 	static bool Load(TSharedPtr<FImplBackEndUEOnly>& InOutImplBackEndUEOnly, const TArray<uint8>& InModelReadFromFileInBytes);
+	static bool Load(TSharedPtr<FImplBackEndUEOnly>& InOutImplBackEndUEOnly, const TArray<uint8>& InModelReadFromFileInBytes);
+	static bool Load(TSharedPtr<FImplBackEndUEOnly>& InOutImplBackEndUEOnly, TSharedPtr<FNeuralTensorManager>& InTensorManager, const TArray<TSharedPtr<FNeuralOperator>>& InOperators);
 
 	//static bool Load(TSharedPtr<FImplBackEndUEOnly>& InOutImplBackEndUEOnly, FNeuralTensorManager& InTensorManager, const TArray<TSharedPtr<FNeuralOperator>>& InOperators);
 
 	void Run(FOnAsyncRunCompleted& InOutOnAsyncRunCompletedDelegate, const ENeuralNetworkSynchronousMode InSynchronousMode, const ENeuralDeviceType InDeviceType, const ENeuralDeviceType InInputDeviceType, const ENeuralDeviceType InOutputDeviceType);
+
+private:
+	/**
+	 *  Used by both Load() functions to reset InOutImplBackEndUEOnly.
+	 */
+	static void Reset(TSharedPtr<FImplBackEndUEOnly>& InOutImplBackEndUEOnly);
 };
diff --git a/Engine/Plugins/Experimental/NeuralNetworkInference/Source/NeuralNetworkInference/Public/NeuralNetwork.h b/Engine/Plugins/Experimental/NeuralNetworkInference/Source/NeuralNetworkInference/Public/NeuralNetwork.h
--- a/Engine/Plugins/Experimental/NeuralNetworkInference/Source/NeuralNetworkInference/Public/NeuralNetwork.h
+++ b/Engine/Plugins/Experimental/NeuralNetworkInference/Source/NeuralNetworkInference/Public/NeuralNetwork.h
@@ -145,11 +145,19 @@ public:
 	void* GetInputDataPointerMutable(const int32 InTensorIndex = 0);
 	int64 GetInputTensorNumber() const;
 
+	/**
+	 * Slow functions (as they will copy every input/output FNeuralNetwork) only meant for debugging purposes.
+	 */
+	TArray<FNeuralTensor> CreateInputArrayCopy() const;
+	void SetInputFromArrayCopy(const TArray<FNeuralTensor>& InTensorDataArray);
+	TArray<FNeuralTensor> CreateOutputArrayCopy() const;
+
 	/**
 	 * Functions to get output. The returned FNeuralTensor(s) are constant to prevent the user from modifying the tensor properties (e.g., size or dimensions).
 	 */
 	const FNeuralTensor& GetOutputTensor(const int32 InTensorIndex = 0) const;
 	int64 GetOutputTensorNumber() const;
+
 	/**
 	 * Non-efficient functions meant to be used only for debugging purposes.
 	 * - InputTensorsToCPU will send the CPU memory of the desired input tensor(s) to GPU memory. Used to debug InputDeviceType == ENeuralDeviceType::GPU.
@@ -256,6 +264,12 @@ private:
 	FNeuralTensor& GetOutputTensorMutable(const int32 InTensorIndex = 0);
 
 public:
+	/**
+	 * Internal function not needed by the user.
+	 * Used to create custom networks without an ONNX file for QA testing in FOperatorTester::TestOperator().
+	 */
+	bool Load(TSharedPtr<struct FNeuralTensorManager>& InTensorManager, const TArray<TSharedPtr<class FNeuralOperator>>& InOperators);
+
 #if WITH_EDITOR
 	/**
 	 * Internal and Editor-only functions not needed by the user.
