NeuralNetworkInference (NNI) plugin:
- NNI-Legacy removed.
- Files removed: Legacy unit tester files, legacy factory.
- Files moved to private: UNNLegacy (to be removed once UNN::UEOnly works in-game)

#rnx

diff --git a/Engine/Plugins/Experimental/NeuralNetworkInference/Source/NeuralNetworkInference/Internal/NeuralNetworkLegacy.h b/Engine/Plugins/Experimental/NeuralNetworkInference/Source/NeuralNetworkInference/Internal/NeuralNetworkLegacy.h
deleted file mode 100644
diff --git a/Engine/Plugins/Experimental/NeuralNetworkInference/Source/NeuralNetworkInference/Private/NeuralNetworkLegacy.cpp b/Engine/Plugins/Experimental/NeuralNetworkInference/Source/NeuralNetworkInference/Private/NeuralNetworkLegacy.cpp
deleted file mode 100644
diff --git a/Engine/Plugins/Experimental/NeuralNetworkInference/Source/NeuralNetworkInferenceEditor/Private/NeuralNetworkInferenceEditorModule.cpp b/Engine/Plugins/Experimental/NeuralNetworkInference/Source/NeuralNetworkInferenceEditor/Private/NeuralNetworkInferenceEditorModule.cpp
--- a/Engine/Plugins/Experimental/NeuralNetworkInference/Source/NeuralNetworkInferenceEditor/Private/NeuralNetworkInferenceEditorModule.cpp
+++ b/Engine/Plugins/Experimental/NeuralNetworkInference/Source/NeuralNetworkInferenceEditor/Private/NeuralNetworkInferenceEditorModule.cpp
@@ -2,7 +2,6 @@
 
 #include "NeuralNetworkInferenceEditorModule.h"
 #include "NeuralNetworkAssetTypeActions.h"
-#include "NeuralNetworkLegacyAssetTypeActions.h"
 #include "IAssetTypeActions.h"
 #include "Interfaces/IPluginManager.h"
 
@@ -22,7 +21,7 @@ public:
 	virtual EAssetTypeCategories::Type GetMLAssetCategoryBit() const;
 
 private:
-	TArray<TSharedPtr<IAssetTypeActions>> Actions;
+	TSharedPtr<IAssetTypeActions> Action;
 	EAssetTypeCategories::Type MLAssetCategoryBit;
 };
 
@@ -36,10 +35,8 @@ void FNeuralNetworkInferenceEditorModule::StartupModule()
 {
 	// UNeuralNetwork - Register asset types
 	IAssetTools& AssetTools = FModuleManager::LoadModuleChecked<FAssetToolsModule>("AssetTools").Get();
-	Actions.Emplace(MakeShared<FNeuralNetworkAssetTypeActions>());
-	AssetTools.RegisterAssetTypeActions(Actions.Last().ToSharedRef());
-	Actions.Emplace(MakeShared<FNeuralNetworkLegacyAssetTypeActions>());
-	AssetTools.RegisterAssetTypeActions(Actions.Last().ToSharedRef());
+	Action = MakeShared<FNeuralNetworkAssetTypeActions>();
+	AssetTools.RegisterAssetTypeActions(Action.ToSharedRef());
 	// Register ML category so that ML assets can register to it
 	MLAssetCategoryBit = AssetTools.RegisterAdvancedAssetCategory(FName(TEXT("ML")), NSLOCTEXT("MLAssetCategory", "MLAssetCategory_ML", "Machine Learning"));
 }
@@ -53,12 +50,9 @@ void FNeuralNetworkInferenceEditorModule::ShutdownModule()
 	{
 		IAssetTools& AssetTools = ModuleInterface->Get();
 		// UNeuralNetwork - Unregister asset types
-		for (TSharedPtr<IAssetTypeActions>& Action : Actions)
+		if (Action.IsValid())
 		{
-			if (Action.IsValid())
-			{
-				AssetTools.UnregisterAssetTypeActions(Action.ToSharedRef());
-			}
+			AssetTools.UnregisterAssetTypeActions(Action.ToSharedRef());
 		}
 	}
 }
diff --git a/Engine/Plugins/Experimental/NeuralNetworkInference/Source/NeuralNetworkInferenceEditor/Private/NeuralNetworkLegacyAssetTypeActions.cpp b/Engine/Plugins/Experimental/NeuralNetworkInference/Source/NeuralNetworkInferenceEditor/Private/NeuralNetworkLegacyAssetTypeActions.cpp
deleted file mode 100644
diff --git a/Engine/Plugins/Experimental/NeuralNetworkInference/Source/NeuralNetworkInferenceEditor/Private/NeuralNetworkLegacyAssetTypeActions.h b/Engine/Plugins/Experimental/NeuralNetworkInference/Source/NeuralNetworkInferenceEditor/Private/NeuralNetworkLegacyAssetTypeActions.h
deleted file mode 100644
diff --git a/Engine/Plugins/Experimental/NeuralNetworkInference/Source/NeuralNetworkInferenceEditor/Private/NeuralNetworkLegacyFactory.cpp b/Engine/Plugins/Experimental/NeuralNetworkInference/Source/NeuralNetworkInferenceEditor/Private/NeuralNetworkLegacyFactory.cpp
deleted file mode 100644
diff --git a/Engine/Plugins/Experimental/NeuralNetworkInference/Source/NeuralNetworkInferenceEditor/Private/NeuralNetworkLegacyFactory.h b/Engine/Plugins/Experimental/NeuralNetworkInference/Source/NeuralNetworkInferenceEditor/Private/NeuralNetworkLegacyFactory.h
deleted file mode 100644
diff --git a/Engine/Plugins/Experimental/NeuralNetworkInference/Source/NeuralNetworkInferenceEditor/Private/NeuralNetworkLegacyReimportFactory.cpp b/Engine/Plugins/Experimental/NeuralNetworkInference/Source/NeuralNetworkInferenceEditor/Private/NeuralNetworkLegacyReimportFactory.cpp
deleted file mode 100644
diff --git a/Engine/Plugins/Experimental/NeuralNetworkInference/Source/NeuralNetworkInferenceEditor/Private/NeuralNetworkLegacyReimportFactory.h b/Engine/Plugins/Experimental/NeuralNetworkInference/Source/NeuralNetworkInferenceEditor/Private/NeuralNetworkLegacyReimportFactory.h
deleted file mode 100644
diff --git a/Engine/Plugins/Experimental/NeuralNetworkInference/Source/NeuralNetworkInferenceQA/Internal/NeuralNetworkInferenceQAAsset.h b/Engine/Plugins/Experimental/NeuralNetworkInference/Source/NeuralNetworkInferenceQA/Internal/NeuralNetworkInferenceQAAsset.h
--- a/Engine/Plugins/Experimental/NeuralNetworkInference/Source/NeuralNetworkInferenceQA/Internal/NeuralNetworkInferenceQAAsset.h
+++ b/Engine/Plugins/Experimental/NeuralNetworkInference/Source/NeuralNetworkInferenceQA/Internal/NeuralNetworkInferenceQAAsset.h
@@ -3,11 +3,10 @@
 #pragma once
 
 #include "CoreMinimal.h"
+#include "NeuralNetwork.h"
 #include "NeuralTensor.h"
 #include "NeuralNetworkInferenceQAAsset.generated.h"
 
-class UNeuralNetworkLegacy;
-
 
 
 USTRUCT(BlueprintType)
@@ -25,8 +24,8 @@ public:
 	/**
 	 * It adds the InNameIndexMap subset of InTensors into Tensors.
 	 */
-	void AddInputTensors(const TArray<FNeuralTensor>& InTensors, const TMap<FString, int32>& InNameIndexMap);
-	void AddOutputTensors(const TArray<FNeuralTensor>& InTensors, const TMap<FString, int32>& InNameIndexMap);
+	void AddInputTensors(const UNeuralNetwork* const InNetwork);
+	void AddOutputTensors(const UNeuralNetwork* const InNetwork);
 
 	/**
 	 * It will compare new vs. previous tests.
@@ -48,9 +47,9 @@ struct NEURALNETWORKINFERENCEQA_API FNeuralNetworkInferenceQAOperatorAsset
 	GENERATED_BODY()
 
 public:
-	static FString RunNetworkCPUAndGetString(UNeuralNetworkLegacy* InOutNetwork);
+	static FString RunNetworkCPUAndGetString(UNeuralNetwork* InOutNetwork);
 
-	void RunAndAddTest(UNeuralNetworkLegacy* InOutNetwork);
+	void RunAndAddTest(UNeuralNetwork* InOutNetwork);
 
 	/**
 	 * It will compare new vs. previous tests.
@@ -77,7 +76,7 @@ protected:
 
 
 /**
- * This is an auxiliary class. See UNeuralNetworkLegacy for a high-level wrapper of the whole NeuralNetworkInference plugin. The UNeuralNetworkLegacy header
+ * This is an auxiliary class. See UNeuralNetwork for a high-level wrapper of the whole NeuralNetworkInference plugin. The UNeuralNetwork header
  * documentation also includes some code examples.
  */
 UCLASS(BlueprintType)
@@ -88,7 +87,7 @@ class NEURALNETWORKINFERENCEQA_API UNeuralNetworkInferenceQAAsset : public UObje
 public:
 	void FindOrAddOperators(const TArray<FString>& InOperatorNames);
 
-	void RunAndAddTest(UNeuralNetworkLegacy* InOutNetwork, const FString& OperatorName);
+	void RunAndAddTest(UNeuralNetwork* InOutNetwork, const FString& OperatorName);
 
 	bool CompareNewVsPreviousTests(const FString& InGroundTruthDirectory);
 
diff --git a/Engine/Plugins/Experimental/NeuralNetworkInference/Source/NeuralNetworkInferenceQA/Private/LegacyModelUnitTester.cpp b/Engine/Plugins/Experimental/NeuralNetworkInference/Source/NeuralNetworkInferenceQA/Private/LegacyModelUnitTester.cpp
deleted file mode 100644
diff --git a/Engine/Plugins/Experimental/NeuralNetworkInference/Source/NeuralNetworkInferenceQA/Private/LegacyModelUnitTester.h b/Engine/Plugins/Experimental/NeuralNetworkInference/Source/NeuralNetworkInferenceQA/Private/LegacyModelUnitTester.h
deleted file mode 100644
diff --git a/Engine/Plugins/Experimental/NeuralNetworkInference/Source/NeuralNetworkInferenceQA/Private/NeuralNetworkInferenceQAAsset.cpp b/Engine/Plugins/Experimental/NeuralNetworkInference/Source/NeuralNetworkInferenceQA/Private/NeuralNetworkInferenceQAAsset.cpp
--- a/Engine/Plugins/Experimental/NeuralNetworkInference/Source/NeuralNetworkInferenceQA/Private/NeuralNetworkInferenceQAAsset.cpp
+++ b/Engine/Plugins/Experimental/NeuralNetworkInference/Source/NeuralNetworkInferenceQA/Private/NeuralNetworkInferenceQAAsset.cpp
@@ -2,7 +2,6 @@
 
 #include "NeuralNetworkInferenceQAAsset.h"
 #include "NeuralNetworkInferenceQAUtils.h"
-#include "NeuralNetworkLegacy.h"
 #include "HAL/FileManager.h"
 #include "Misc/FileHelper.h"
 #include "Misc/Paths.h"
@@ -12,19 +11,19 @@
 /* FNeuralNetworkInferenceQAOperatorTestAsset public functions
  *****************************************************************************/
 
-void FNeuralNetworkInferenceQAOperatorTestAsset::AddInputTensors(const TArray<FNeuralTensor>& InTensors, const TMap<FString, int32>& InNameIndexMap)
+void FNeuralNetworkInferenceQAOperatorTestAsset::AddInputTensors(const UNeuralNetwork* const InNetwork)
 {
-	for (const auto& NameIndexPair : InNameIndexMap)
+	for (int32 InputIndex = 0; InputIndex < InNetwork->GetInputTensorNumber(); ++InputIndex)
 	{
-		InputTensors.Push(InTensors[NameIndexPair.Value]);
+		InputTensors.Push(InNetwork->GetInputTensor(InputIndex));
 	}
 }
 
-void FNeuralNetworkInferenceQAOperatorTestAsset::AddOutputTensors(const TArray<FNeuralTensor>& InTensors, const TMap<FString, int32>& InNameIndexMap)
+void FNeuralNetworkInferenceQAOperatorTestAsset::AddOutputTensors(const UNeuralNetwork* const InNetwork)
 {
-	for (const auto& NameIndexPair : InNameIndexMap)
+	for (int32 OutputIndex = 0; OutputIndex < InNetwork->GetOutputTensorNumber(); ++OutputIndex)
 	{
-		OutputTensors.Push(InTensors[NameIndexPair.Value]);
+		OutputTensors.Push(InNetwork->GetOutputTensor(OutputIndex));
 	}
 }
 
@@ -59,41 +58,39 @@ bool FNeuralNetworkInferenceQAOperatorTestAsset::CompareAverageL1DiffNewVsPrevio
 /* FNeuralNetworkInferenceQAOperatorAsset public functions
  *****************************************************************************/
 
-FString FNeuralNetworkInferenceQAOperatorAsset::RunNetworkCPUAndGetString(UNeuralNetworkLegacy* InOutNetwork)
+FString FNeuralNetworkInferenceQAOperatorAsset::RunNetworkCPUAndGetString(UNeuralNetwork* InOutNetwork)
 {
-	const TArray<FNeuralTensor>& Tensors = InOutNetwork->GetTensors();
-	// Create input string and InputTensorMap
-	const TMap<FString, FNeuralTensor> InputTensorMap = InOutNetwork->CreateInputTensorMap();
+	// Create input string and InputTensorArray
+	const TArray<FNeuralTensor> InputTensorArray = InOutNetwork->CreateInputArrayCopy();
 	FString TensorsAsString = TEXT("Input(s):\n");
-	for (const auto& InputTensor : InOutNetwork->GetInputNameIndexMap())
+	for (int32 InputIndex = 0; InputIndex < InOutNetwork->GetInputTensorNumber(); ++InputIndex)
 	{
-		TensorsAsString += Tensors[InputTensor.Value].ToString() + TEXT("\n");
+		TensorsAsString += InOutNetwork->GetInputTensor(InputIndex).ToString() + TEXT("\n");
 	}
 	// Run network
 	InOutNetwork->SetDeviceType(ENeuralDeviceType::CPU);
 	InOutNetwork->Run();
 	// Create output string
-	const TMap<FString, int32>& OutputNameIndexMap = InOutNetwork->GetOutputNameIndexMap();
 	TensorsAsString += TEXT("\nOutput(s):\n");
-	for (const auto& OutputTensorPair : OutputNameIndexMap)
+	for (int32 OutputIndex = 0; OutputIndex < InOutNetwork->GetOutputTensorNumber(); ++OutputIndex)
 	{
-		TensorsAsString += Tensors[OutputTensorPair.Value].ToString();
+		TensorsAsString += InOutNetwork->GetOutputTensor(OutputIndex).ToString() + TEXT("\n");
 	}
 	// Reset memory
-	InOutNetwork->SetInputFromTensorMapCopy(InputTensorMap); // This is doing an unnecessary copy (it should move really)
+	InOutNetwork->SetInputFromArrayCopy(InputTensorArray); // This is doing an unnecessary copy (it should move really)
 	// Return TensorsAsString
 	return TensorsAsString + TEXT("\n\n\n");
 }
 
-void FNeuralNetworkInferenceQAOperatorAsset::RunAndAddTest(UNeuralNetworkLegacy* InOutNetwork)
+void FNeuralNetworkInferenceQAOperatorAsset::RunAndAddTest(UNeuralNetwork* InOutNetwork)
 {
 	// Resize NewTests
 	NewTests.Emplace();
 
 	FNeuralNetworkInferenceQAOperatorTestAsset& NewTest = NewTests.Last();
 	NewTestsString += RunNetworkCPUAndGetString(InOutNetwork);
-	NewTest.AddInputTensors(InOutNetwork->GetTensors(), InOutNetwork->GetInputNameIndexMap());
-	NewTest.AddOutputTensors(InOutNetwork->GetTensors(), InOutNetwork->GetOutputNameIndexMap());
+	NewTest.AddInputTensors(InOutNetwork);
+	NewTest.AddOutputTensors(InOutNetwork);
 }
 
 bool FNeuralNetworkInferenceQAOperatorAsset::CompareNewVsPreviousTests(const FString& InGroundTruthDirectory, const FString& InOperatorName)
@@ -183,7 +180,7 @@ void UNeuralNetworkInferenceQAAsset::FindOrAddOperators(const TArray<FString>& I
 	}
 }
 
-void UNeuralNetworkInferenceQAAsset::RunAndAddTest(UNeuralNetworkLegacy* InOutNetwork, const FString& OperatorName)
+void UNeuralNetworkInferenceQAAsset::RunAndAddTest(UNeuralNetwork* InOutNetwork, const FString& OperatorName)
 {
 	Operators[OperatorName].RunAndAddTest(InOutNetwork);
 }
diff --git a/Engine/Plugins/Experimental/NeuralNetworkInference/Source/NeuralNetworkInferenceQA/Private/OperatorTester.cpp b/Engine/Plugins/Experimental/NeuralNetworkInference/Source/NeuralNetworkInferenceQA/Private/OperatorTester.cpp
--- a/Engine/Plugins/Experimental/NeuralNetworkInference/Source/NeuralNetworkInferenceQA/Private/OperatorTester.cpp
+++ b/Engine/Plugins/Experimental/NeuralNetworkInference/Source/NeuralNetworkInferenceQA/Private/OperatorTester.cpp
@@ -3,6 +3,7 @@
 #include "OperatorTester.h"
 #include "NeuralNetworkInferenceQAUtils.h"
 #include "NeuralOperators.h"
+#include "NeuralTensorManager.h"
 #include "Misc/Paths.h"
 
 
@@ -75,14 +76,14 @@ void FOperatorTester::TestOperator(UNeuralNetworkInferenceQAAsset* InOutNetworkI
 			return;
 		}
 		// Network tests
-		UNeuralNetworkLegacy* Network = NewObject<UNeuralNetworkLegacy>((UObject*)GetTransientPackage(), UNeuralNetworkLegacy::StaticClass());
+		UNeuralNetwork* Network = NewObject<UNeuralNetwork>((UObject*)GetTransientPackage(), UNeuralNetwork::StaticClass());
 		if (!Network)
 		{
 			ensureMsgf(false, TEXT("Network is nullptr."));
 			return;
 		}
-		FNeuralTensorManager TensorManager(InOutTensors[Index], InputTensors, OutputTensors);
-		if (!TensorManager.IsLoaded())
+		TSharedPtr<FNeuralTensorManager> TensorManager = MakeShared<FNeuralTensorManager>(InOutTensors[Index], InputTensors, OutputTensors);
+		if (!TensorManager->IsLoaded())
 		{
 			ensureMsgf(false, TEXT("TensorManager could not be loaded."));
 			return;
@@ -100,7 +101,7 @@ void FOperatorTester::TestOperator(UNeuralNetworkInferenceQAAsset* InOutNetworkI
 		Network->SetDeviceType(ENeuralDeviceType::CPU);
 		if (InlinedTensorIndex > -1)
 		{
-			NetworkInputVsOutput(Network);
+			NetworkInputVsOutput(Network, InlinedTensorIndex);
 		}
 		// Test against GT saved on disk
 		// Case a) Only saving first test, the others should match the first test results
@@ -1161,74 +1162,70 @@ TArray<int64> FOperatorTester::FloatTensorToIntArray(const TArray<FNeuralTensor>
 	return OutputArray;
 }
 
-void FOperatorTester::ResetOrCheckInput(UNeuralNetworkLegacy* InOutNetwork, const TMap<FString, FNeuralTensor>& InInputTensorMap, const bool bIsInlinedTensor)
+void FOperatorTester::ResetOrCheckInput(UNeuralNetwork* InOutNetwork, const TArray<FNeuralTensor>& InInputTensorArray, const bool bIsInlinedTensor)
 {
 	if (bIsInlinedTensor)
 	{
-		InOutNetwork->SetInputFromTensorMapCopy(InInputTensorMap);
+		InOutNetwork->SetInputFromArrayCopy(InInputTensorArray);
 	}
 	// Check input has not changed
 	else
 	{
-		const TArray<FNeuralTensor>& Tensors = InOutNetwork->GetTensors();
-		const TMap<FString, int32>& InputNameIndexMap = InOutNetwork->GetInputNameIndexMap();
-		for (const auto& InputTensorTuple : InInputTensorMap)
+		for (int32 InputIndex = 0; InputIndex < InOutNetwork->GetInputTensorNumber(); ++InputIndex)
 		{
-			const FNeuralTensor& InputTensor = Tensors[InputNameIndexMap.FindChecked(InputTensorTuple.Key)];
-			ensureMsgf(InputTensorTuple.Value.Num() == InputTensor.Num(), TEXT("InputTensorTuple.Value.Num() == InputTensor.Num() failed (%d != %d)."),
-				InputTensorTuple.Value.Num(), InputTensor.Num());
-			ensureMsgf(FNeuralNetworkInferenceQAUtils::EstimateTensorL1DiffError(InputTensor, InputTensorTuple.Value, 0.f, TEXT("ResetOrCheckInput")), TEXT("ResetOrCheckInput() did not pass."));
+			const FNeuralTensor& InputTensor = InOutNetwork->GetInputTensor(InputIndex);
+			ensureMsgf(InInputTensorArray[InputIndex].Num() == InputTensor.Num(), TEXT("InInputTensorArray[InputIndex].Num() == InputTensor.Num() failed (%d != %d)."),
+				InInputTensorArray[InputIndex].Num(), InputTensor.Num());
+			ensureMsgf(FNeuralNetworkInferenceQAUtils::EstimateTensorL1DiffError(InputTensor, InInputTensorArray[InputIndex], 0.f, TEXT("ResetOrCheckInput")), TEXT("ResetOrCheckInput() did not pass."));
 		}
 	}
 }
 
-void FOperatorTester::NetworkGPUvsCPU(UNeuralNetworkLegacy* InOutNetwork, const bool bIsInlinedTensor, const float InZeroThreshold)
+void FOperatorTester::NetworkGPUvsCPU(UNeuralNetwork* InOutNetwork, const bool bIsInlinedTensor, const float InZeroThreshold)
 {
-	// Create InputTensorMap
-	TMap<FString, FNeuralTensor> InputTensorMap = InOutNetwork->CreateInputTensorMap();
+	// Create InputTensorDataArray
+	const TArray<FNeuralTensor> InputTensorArray = InOutNetwork->CreateInputArrayCopy();
 	// GPU
 	InOutNetwork->SetDeviceType(ENeuralDeviceType::GPU);
 	InOutNetwork->Run();
 	// Store GPU results
-	const TMap<FString, FNeuralTensor> OutputTensorMapGPU = InOutNetwork->CreateOutputTensorMap();
+	const TArray<FNeuralTensor> OutputTensorArray = InOutNetwork->CreateOutputArrayCopy();
 	// Reset memory / Check input has not changed
-	ResetOrCheckInput(InOutNetwork, InputTensorMap, bIsInlinedTensor);
+	ResetOrCheckInput(InOutNetwork, InputTensorArray, bIsInlinedTensor);
 	// CPU
 	InOutNetwork->SetDeviceType(ENeuralDeviceType::CPU);
 	InOutNetwork->Run();
 	// CPU vs. GPU
-	const TArray<FNeuralTensor>& Tensors = InOutNetwork->GetTensors();
-	const TMap<FString, int32>& OutputNameIndexMap = InOutNetwork->GetOutputNameIndexMap();
-	for (const auto& OutputTensorGPUTuple : OutputTensorMapGPU)
+	for (int32 OutputIndex = 0; OutputIndex < InOutNetwork->GetOutputTensorNumber(); ++OutputIndex)
 	{
-		const FNeuralTensor& OutputTensor = Tensors[OutputNameIndexMap.FindChecked(OutputTensorGPUTuple.Key)];
-		ensureMsgf(OutputTensorGPUTuple.Value.Num() == OutputTensor.Num(), TEXT("OutputTensorGPUTuple.Value.Num() == OutputTensor.Num(): %d vs. %d"),
-			OutputTensorGPUTuple.Value.Num(), OutputTensor.Num());
-		ensureMsgf(FNeuralNetworkInferenceQAUtils::EstimateTensorL1DiffError(OutputTensor, OutputTensorGPUTuple.Value, InZeroThreshold, TEXT("GPUvsCPU")), TEXT("CPU vs. GPU error has too high"));
+		const FNeuralTensor& OutputTensorCPU = OutputTensorArray[OutputIndex];
+		const FNeuralTensor& OutputTensorGPU = InOutNetwork->GetOutputTensor(OutputIndex);
+		ensureMsgf(OutputTensorCPU.Num() == OutputTensorGPU.Num(), TEXT("OutputTensorCPU.Num() == OutputTensorGPU.Num(): %d vs. %d"),
+			OutputTensorCPU.Num(), OutputTensorGPU.Num());
+		ensureMsgf(FNeuralNetworkInferenceQAUtils::EstimateTensorL1DiffError(OutputTensorCPU, OutputTensorGPU, InZeroThreshold, TEXT("GPUvsCPU")), TEXT("CPU vs. GPU error has too high"));
 	}
 	// Reset memory / Check input has not changed
-	ResetOrCheckInput(InOutNetwork, InputTensorMap, bIsInlinedTensor); // This is doing an unnecessary copy (it should move really)
+	ResetOrCheckInput(InOutNetwork, InputTensorArray, bIsInlinedTensor); // This is doing an unnecessary copy (it should move really)
 }
 
-void FOperatorTester::NetworkInputVsOutput(UNeuralNetworkLegacy* InOutNetwork)
+void FOperatorTester::NetworkInputVsOutput(UNeuralNetwork* InOutNetwork, const int32 InInlinedTensorIndex)
 {
-	TMap<FString, FNeuralTensor> InputTensorMap = InOutNetwork->CreateInputTensorMap();
-	// Run UNeuralNetworkLegacy
+	const TArray<FNeuralTensor> InputTensorArray = InOutNetwork->CreateInputArrayCopy();
+	// Run UNeuralNetwork
 	InOutNetwork->Run();
-	// Check input equals output
-	for (const auto& OutputTensorPair : InOutNetwork->GetOutputNameIndexMap())
+	// Check input equals one of the outputs
+	const FNeuralTensor& InputTensor = InOutNetwork->GetInputTensor(InInlinedTensorIndex);
+	bool bIsSomeInputAnOutput = false;
+	for (int32 OutputIndex = 0; OutputIndex < InOutNetwork->GetOutputTensorNumber(); ++OutputIndex)
 	{
-		bool bIsSomeInputAnOutput = false;
-		for (const auto& InputTensorPair : InOutNetwork->GetInputNameIndexMap())
+		const FNeuralTensor& OutputTensor = InOutNetwork->GetOutputTensor(OutputIndex);
+		if (InputTensor.GetUnderlyingUInt8ArrayRef() == OutputTensor.GetUnderlyingUInt8ArrayRef())
 		{
-			if (InputTensorPair.Value == OutputTensorPair.Value)
-			{
-				bIsSomeInputAnOutput = true;
-				break;
-			}
+			bIsSomeInputAnOutput = true;
+			break;
 		}
-		ensureMsgf(bIsSomeInputAnOutput, TEXT("No input is being used as an output."));
 	}
+	ensureMsgf(bIsSomeInputAnOutput, TEXT("Input is not being used as an output."));
 	// Reset memory
-	InOutNetwork->SetInputFromTensorMapCopy(InputTensorMap); // This is doing an unnecessary copy (it should move really)
+	InOutNetwork->SetInputFromArrayCopy(InputTensorArray); // This is doing an unnecessary copy (it should move really)
 }
diff --git a/Engine/Plugins/Experimental/NeuralNetworkInference/Source/NeuralNetworkInferenceQA/Private/OperatorTester.h b/Engine/Plugins/Experimental/NeuralNetworkInference/Source/NeuralNetworkInferenceQA/Private/OperatorTester.h
--- a/Engine/Plugins/Experimental/NeuralNetworkInference/Source/NeuralNetworkInferenceQA/Private/OperatorTester.h
+++ b/Engine/Plugins/Experimental/NeuralNetworkInference/Source/NeuralNetworkInferenceQA/Private/OperatorTester.h
@@ -4,7 +4,7 @@
 
 #include "CoreMinimal.h"
 #include "NeuralEnumClasses.h"
-#include "NeuralNetworkLegacy.h"
+#include "NeuralNetwork.h"
 #include "NeuralNetworkInferenceQAUtils.h"
 #include "NeuralNetworkInferenceQAAsset.h"
 #include "NeuralOperator.h"
@@ -58,24 +58,24 @@ private:
 	static TArray<int64> FloatTensorToIntArray(const TArray<FNeuralTensor>& InTensors, const int32 InTensorIndex);
 
 	/**
-	 * For inlined layers, it resets the memory (meant after a UNeuralNetworkLegacy::Run()). For non-inlined layers, it checks that the input did not change.
+	 * For inlined layers, it resets the memory (meant after a UNeuralNetwork::Run()). For non-inlined layers, it checks that the input did not change.
 	 */
-	static void ResetOrCheckInput(UNeuralNetworkLegacy* InOutNetwork, const TMap<FString, FNeuralTensor>& InInputTensorMap, const bool bIsInlinedTensor);
+	static void ResetOrCheckInput(UNeuralNetwork* InOutNetwork, const TArray<FNeuralTensor>& InInputTensorArray, const bool bIsInlinedTensor);
 
 	/**
 	 * It tests and makes sure CPU and GPU modes provide the same results. Otherwise, it will check to false.
-	 * @param InOutNetwork UNeuralNetworkLegacy to test.
+	 * @param InOutNetwork UNeuralNetwork to test.
 	 * @param If not inlined, it will make check whether the input tensors have not changed after running.
 	 * @param InZeroThreshold If the average L2 norm between CPU and GPU results is greater than this number, it will check to false.
 	 */
-	static void NetworkGPUvsCPU(UNeuralNetworkLegacy* InOutNetwork, const bool bIsInlinedTensor, const float InZeroThreshold = 0.f);
+	static void NetworkGPUvsCPU(UNeuralNetwork* InOutNetwork, const bool bIsInlinedTensor, const float InZeroThreshold = 0.f);
 
 	/**
 	 * If layer inlined --> Check input == output
 	 * This is only meant for networks with single input single output, whose layers are all inlined (so that input tensor data should match
-	 * output tensor data after a UNeuralNetworkLegacy run).
+	 * output tensor data after a UNeuralNetwork run).
 	 */
-	static void NetworkInputVsOutput(UNeuralNetworkLegacy* InOutNetwork);
+	static void NetworkInputVsOutput(UNeuralNetwork* InOutNetwork, const int32 InInlinedTensorIndex);
 };
 
 
diff --git a/Engine/Plugins/Experimental/NeuralNetworkInference/Source/NeuralNetworkInferenceQA/Private/UnitTester.cpp b/Engine/Plugins/Experimental/NeuralNetworkInference/Source/NeuralNetworkInferenceQA/Private/UnitTester.cpp
--- a/Engine/Plugins/Experimental/NeuralNetworkInference/Source/NeuralNetworkInferenceQA/Private/UnitTester.cpp
+++ b/Engine/Plugins/Experimental/NeuralNetworkInference/Source/NeuralNetworkInferenceQA/Private/UnitTester.cpp
@@ -1,17 +1,10 @@
 // Copyright Epic Games, Inc. All Rights Reserved.
 
 #include "UnitTester.h"
-#include "LegacyModelUnitTester.h"
 #include "ModelUnitTester.h"
 #include "NeuralNetworkInferenceQAUtils.h"
 #include "OperatorUnitTester.h"
 
-// #if WITH_EDITOR
-// #ifdef PLATFORM_WIN64
-// #include "ONNXRuntimeDLLTester.h"
-// #endif //PLATFORM_WIN64
-// #endif //WITH_EDITOR
-
 
 
 /* FUnitTester static public functions
@@ -24,31 +17,18 @@ bool FUnitTester::GlobalTest(const FString& InProjectContentDir, const FString&
 	UE_LOG(LogNeuralNetworkInferenceQA, Display, TEXT("---------------------------------------------------------------------------------------------------------------------------------"));
 	UE_LOG(LogNeuralNetworkInferenceQA, Display, TEXT("----- Starting UnitTesting() ----------------------------------------------------------------------------------------------------"));
 
-	UE_LOG(LogNeuralNetworkInferenceQA, Display, TEXT("---------- 1. Model Unit Testing (Legacy)"));
-	bool bDidGlobalTestPassed = FLegacyModelUnitTester::GlobalTest(InProjectContentDir, InModelZooRelativeDirectory);
-
 	UE_LOG(LogNeuralNetworkInferenceQA, Display, TEXT("---------------------------------------------------------------------------------------------------------------------------------"));
 	UE_LOG(LogNeuralNetworkInferenceQA, Display, TEXT("---------------------------------------------------------------------------------------------------------------------------------"));
 	UE_LOG(LogNeuralNetworkInferenceQA, Display, TEXT("---------------------------------------------------------------------------------------------------------------------------------"));
-	UE_LOG(LogNeuralNetworkInferenceQA, Display, TEXT("---------- 2. Model Unit Testing"));
-	bDidGlobalTestPassed &= FModelUnitTester::GlobalTest(InProjectContentDir, InModelZooRelativeDirectory);
+	UE_LOG(LogNeuralNetworkInferenceQA, Display, TEXT("---------- 1. Model Unit Testing"));
+	bool bDidGlobalTestPassed = FModelUnitTester::GlobalTest(InProjectContentDir, InModelZooRelativeDirectory);
 
 	UE_LOG(LogNeuralNetworkInferenceQA, Display, TEXT("---------------------------------------------------------------------------------------------------------------------------------"));
 	UE_LOG(LogNeuralNetworkInferenceQA, Display, TEXT("---------------------------------------------------------------------------------------------------------------------------------"));
 	UE_LOG(LogNeuralNetworkInferenceQA, Display, TEXT("---------------------------------------------------------------------------------------------------------------------------------"));
-	UE_LOG(LogNeuralNetworkInferenceQA, Display, TEXT("---------- 3. Operator Unit Testing"));
+	UE_LOG(LogNeuralNetworkInferenceQA, Display, TEXT("---------- 2. Operator Unit Testing"));
 	bDidGlobalTestPassed &= FOperatorUnitTester::GlobalTest(InProjectContentDir, InUnitTestRelativeDirectory);
 
-// #if WITH_EDITOR
-// #ifdef PLATFORM_WIN64
-// 	UE_LOG(LogNeuralNetworkInferenceQA, Display, TEXT("---------------------------------------------------------------------------------------------------------------------------------"));
-// 	UE_LOG(LogNeuralNetworkInferenceQA, Display, TEXT("---------------------------------------------------------------------------------------------------------------------------------"));
-// 	UE_LOG(LogNeuralNetworkInferenceQA, Display, TEXT("---------------------------------------------------------------------------------------------------------------------------------"));
-// 	UE_LOG(LogNeuralNetworkInferenceQA, Display, TEXT("---------- 4. ONNX Runtime DLL Unit Testing (Deprecated)"));
-// 	bDidGlobalTestPassed &= FONNXRuntimeDLLTester::GlobalTest(InProjectContentDir, InModelZooRelativeDirectory);
-// #endif //PLATFORM_WIN64
-// #endif //WITH_EDITOR
-
 	if (bDidGlobalTestPassed)
 	{
 		UE_LOG(LogNeuralNetworkInferenceQA, Display, TEXT("----- UnitTesting() completed! --------------------------------------------------------------------------------------------------"));
