- Move FNeuralTensorManager from NNI/Internal into Private now that UNNLegacy is gone.
- FNeuralTensorManager: No longer NNI_API nor struct nor UStruct.
- New constructor in UNN to avoid using FNeuralTensorManager in its header.
- FNeuralTensorManager used as reference (rather than shared ptr) in UNNBackEndUEOnly now that only used in Private.

_codereview Francisco.Vicente
#rnx

diff --git a/Engine/Plugins/Experimental/NeuralNetworkInference/Source/NeuralNetworkInference/Internal/NeuralTensorManager.h b/Engine/Plugins/Experimental/NeuralNetworkInference/Source/NeuralNetworkInference/Internal/NeuralTensorManager.h
deleted file mode 100644
diff --git a/Engine/Plugins/Experimental/NeuralNetworkInference/Source/NeuralNetworkInference/Private/GraphProtoToNeuralNetworkConverter.cpp b/Engine/Plugins/Experimental/NeuralNetworkInference/Source/NeuralNetworkInference/Private/GraphProtoToNeuralNetworkConverter.cpp
deleted file mode 100644
diff --git a/Engine/Plugins/Experimental/NeuralNetworkInference/Source/NeuralNetworkInference/Private/GraphProtoToNeuralNetworkConverter.h b/Engine/Plugins/Experimental/NeuralNetworkInference/Source/NeuralNetworkInference/Private/GraphProtoToNeuralNetworkConverter.h
deleted file mode 100644
diff --git a/Engine/Plugins/Experimental/NeuralNetworkInference/Source/NeuralNetworkInference/Private/NeuralNetwork.cpp b/Engine/Plugins/Experimental/NeuralNetworkInference/Source/NeuralNetworkInference/Private/NeuralNetwork.cpp
--- a/Engine/Plugins/Experimental/NeuralNetworkInference/Source/NeuralNetworkInference/Private/NeuralNetwork.cpp
+++ b/Engine/Plugins/Experimental/NeuralNetworkInference/Source/NeuralNetworkInference/Private/NeuralNetwork.cpp
@@ -228,7 +228,7 @@ bool UNeuralNetwork::IsGPUConfigCompatibleForCurrentBackEnd() const
 	/* UEOnly */ \
 	else if (BackEndForCurrentPlatform == ENeuralBackEnd::UEOnly) \
 	{ \
-		FNeuralTensorManager& TensorManager = *ImplBackEndUEOnly->TensorManager; \
+		FNeuralTensorManager& TensorManager = ImplBackEndUEOnly->TensorManager; \
 		return TensorManager.UEOnly_GetTensorsFunction()[TensorManager.UEAndORTOnly_GetIndexes()[InTensorIndex]]; \
 	} \
 	/* Unknown */ \
@@ -257,7 +257,7 @@ void UNeuralNetwork::SetInputFromArrayCopy(const TArray<float>& InArray, const i
 	// UEOnly
 	else if (BackEndForCurrentPlatform == ENeuralBackEnd::UEOnly)
 	{
-		FNeuralTensorManager& TensorManager = *ImplBackEndUEOnly->TensorManager;
+		FNeuralTensorManager& TensorManager = ImplBackEndUEOnly->TensorManager;
 		TensorManager.GetTensorsMutable()[TensorManager.GetInputIndexes()[InTensorIndex]].SetFromArrayCopy(InArray);
 	}
 
@@ -286,7 +286,7 @@ void* UNeuralNetwork::GetInputDataPointerMutable(const int32 InTensorIndex)
 	// UEOnly
 	else if (BackEndForCurrentPlatform == ENeuralBackEnd::UEOnly)
 	{
-		FNeuralTensorManager& TensorManager = *ImplBackEndUEOnly->TensorManager;
+		FNeuralTensorManager& TensorManager = ImplBackEndUEOnly->TensorManager;
 		return TensorManager.GetTensorsMutable()[TensorManager.GetInputIndexes()[InTensorIndex]].GetData();
 	}
 
@@ -313,7 +313,7 @@ int64 UNeuralNetwork::GetInputTensorNumber() const
 	// UEOnly
 	else if (BackEndForCurrentPlatform == ENeuralBackEnd::UEOnly)
 	{
-		return ImplBackEndUEOnly->TensorManager->GetInputIndexes().Num();
+		return ImplBackEndUEOnly->TensorManager.GetInputIndexes().Num();
 	}
 
 	// Unknown
@@ -344,7 +344,7 @@ int64 UNeuralNetwork::GetOutputTensorNumber() const
 	// UEOnly
 	else if (BackEndForCurrentPlatform == ENeuralBackEnd::UEOnly)
 	{
-		return ImplBackEndUEOnly->TensorManager->GetOutputIndexes().Num();
+		return ImplBackEndUEOnly->TensorManager.GetOutputIndexes().Num();
 	}
 
 	// Unknown
@@ -544,12 +544,20 @@ FNeuralTensor& UNeuralNetwork::GetOutputTensorMutable(const int32 InTensorIndex)
 	GET_TENSOR_CODE(InTensorIndex, GetOutputIndexes, GetTensorsMutable, OutputTensors);
 }
 
-bool UNeuralNetwork::Load(TSharedPtr<FNeuralTensorManager>& InTensorManager, const TArray<TSharedPtr<FNeuralOperator>>& InOperators)
+bool UNeuralNetwork::Load(TArray<FNeuralTensor>& InTensors, const TArray<FNeuralTensor*>& InInputTensors, const TArray<FNeuralTensor*>& InOutputTensors, const TArray<TSharedPtr<class FNeuralOperator>>& InOperators)
 {
 	DECLARE_SCOPE_CYCLE_COUNTER(TEXT("UNeuralNetwork_Load_FromTensorManagerAndOperators"), STAT_UNeuralNetwork_Load, STATGROUP_MachineLearning);
 	BackEnd = ENeuralBackEnd::UEOnly;
 	BackEndForCurrentPlatform = ENeuralBackEnd::UEOnly;
-	bIsLoaded = UNeuralNetwork::FImplBackEndUEOnly::Load(ImplBackEndUEOnly, InTensorManager, InOperators);
+	// Create and load TensorManager
+	FNeuralTensorManager TensorManager(InTensors, InInputTensors, InOutputTensors);
+	if (!TensorManager.IsLoaded())
+	{
+		UE_LOG(LogNeuralNetworkInference, Warning, TEXT("UNeuralNetwork::Load(): TensorManager could not be loaded."));
+		return false;
+	}
+	// Load network
+	bIsLoaded = UNeuralNetwork::FImplBackEndUEOnly::Load(ImplBackEndUEOnly, TensorManager, InOperators);
 	return bIsLoaded;
 }
 
diff --git a/Engine/Plugins/Experimental/NeuralNetworkInference/Source/NeuralNetworkInference/Private/NeuralNetworkImplBackEndUEOnly.cpp b/Engine/Plugins/Experimental/NeuralNetworkInference/Source/NeuralNetworkInference/Private/NeuralNetworkImplBackEndUEOnly.cpp
--- a/Engine/Plugins/Experimental/NeuralNetworkInference/Source/NeuralNetworkInference/Private/NeuralNetworkImplBackEndUEOnly.cpp
+++ b/Engine/Plugins/Experimental/NeuralNetworkInference/Source/NeuralNetworkInference/Private/NeuralNetworkImplBackEndUEOnly.cpp
@@ -1,11 +1,11 @@
 // Copyright Epic Games, Inc. All Rights Reserved.
 
 #include "NeuralNetworkImplBackEndUEOnly.h"
-#include "GraphProtoToNeuralNetworkConverter.h"
 #include "ModelProtoFileReader.h"
 #include "NeuralNetworkInferenceUtils.h"
 #include "RenderGraphBuilder.h"
 #include "RenderingThread.h"
+#include "UEOnly/GraphProtoToNeuralNetworkConverter.h"
 #include "RHI.h"
 
 
@@ -29,37 +29,32 @@ bool UNeuralNetwork::FImplBackEndUEOnly::Load(TSharedPtr<FImplBackEndUEOnly>& In
 		return false;
 	}
 	// Turn ModelProto into Operators
-	if (!FGraphProtoToNeuralNetworkConverter::Translate(InOutImplBackEndUEOnly->Operators, *InOutImplBackEndUEOnly->TensorManager, InOutImplBackEndUEOnly->ModelProto.GetGraph(), /*bInIsTensorManagerConst*/false))
+	if (!FGraphProtoToNeuralNetworkConverter::Translate(InOutImplBackEndUEOnly->Operators, InOutImplBackEndUEOnly->TensorManager, InOutImplBackEndUEOnly->ModelProto.GetGraph(), /*bInIsTensorManagerConst*/false))
 	{
 		UE_LOG(LogNeuralNetworkInference, Warning, TEXT("FImplBackEndUEOnly::Load(): Translate failed."));
 		return false;
 	}
-	if (!InOutImplBackEndUEOnly->TensorManager->IsLoaded())
+	if (!InOutImplBackEndUEOnly->TensorManager.IsLoaded())
 	{
-		UE_LOG(LogNeuralNetworkInference, Warning, TEXT("FImplBackEndUEOnly::Load(): TensorManager->IsLoaded() returned false."));
+		UE_LOG(LogNeuralNetworkInference, Warning, TEXT("FImplBackEndUEOnly::Load(): TensorManager.IsLoaded() returned false."));
 		return false;
 	}
 	// Load successful
 	return true;
 }
 
-bool UNeuralNetwork::FImplBackEndUEOnly::Load(TSharedPtr<FImplBackEndUEOnly>& InOutImplBackEndUEOnly, TSharedPtr<FNeuralTensorManager>& InTensorManager, const TArray<TSharedPtr<FNeuralOperator>>& InOperators)
+bool UNeuralNetwork::FImplBackEndUEOnly::Load(TSharedPtr<FImplBackEndUEOnly>& InOutImplBackEndUEOnly, FNeuralTensorManager& InTensorManager, const TArray<TSharedPtr<FNeuralOperator>>& InOperators)
 {
 	// Reset InOutImplBackEndUEOnly
 	Reset(InOutImplBackEndUEOnly);
 	// Load
-	if (!InTensorManager.IsValid())
-	{
-		UE_LOG(LogNeuralNetworkInference, Warning, TEXT("UNeuralNetworkLegacy::Load(): TensorManager was a nullptr."));
-		return false;
-	}
-	if (!InTensorManager->IsLoaded())
+	if (!InTensorManager.IsLoaded())
 	{
 		UE_LOG(LogNeuralNetworkInference, Warning, TEXT("UNeuralNetworkLegacy::Load(): TensorManager could not be loaded."));
 	}
 	Swap(InOutImplBackEndUEOnly->TensorManager, InTensorManager);
 	InOutImplBackEndUEOnly->Operators = InOperators;
-	return (InOutImplBackEndUEOnly->Operators.Num() > 0 && InOutImplBackEndUEOnly->TensorManager->IsLoaded());
+	return (InOutImplBackEndUEOnly->Operators.Num() > 0 && InOutImplBackEndUEOnly->TensorManager.IsLoaded());
 }
 
 //bool UNeuralNetwork::FImplBackEndUEOnly::Load(TSharedPtr<FImplBackEndUEOnly>& InOutImplBackEndUEOnly, FNeuralTensorManager& InTensorManager, const TArray<TSharedPtr<FNeuralOperator>>& InOperators)
@@ -76,7 +71,7 @@ bool UNeuralNetwork::FImplBackEndUEOnly::Load(TSharedPtr<FImplBackEndUEOnly>& In
 //	}
 //	Swap(InOutImplBackEndUEOnly->TensorManager, InTensorManager);
 //	InOutImplBackEndUEOnly->Operators = InOperators;
-//	return (InOutImplBackEndUEOnly->Operators.Num() > 0 && InOutImplBackEndUEOnly->TensorManager->IsLoaded());
+//	return (InOutImplBackEndUEOnly->Operators.Num() > 0 && InOutImplBackEndUEOnly->TensorManager.IsLoaded());
 //}
 
 void UNeuralNetwork::FImplBackEndUEOnly::Run(FOnAsyncRunCompleted& InOutOnAsyncRunCompletedDelegate, const ENeuralNetworkSynchronousMode InSynchronousMode, const ENeuralDeviceType InDeviceType, const ENeuralDeviceType InInputDeviceType, const ENeuralDeviceType InOutputDeviceType)
@@ -103,9 +98,9 @@ void UNeuralNetwork::FImplBackEndUEOnly::Run(FOnAsyncRunCompleted& InOutOnAsyncR
 		{
 			DECLARE_SCOPE_CYCLE_COUNTER(TEXT("UNeuralNetwork_UEOnly_Run::Forward_GPU"), STAT_UNeuralNetwork_UEOnly_Run_Forward_GPU, STATGROUP_MachineLearning);
 			// Sanity check
-			if (TensorManager->GetTensorsMutable().Num() < 1)
+			if (TensorManager.GetTensorsMutable().Num() < 1)
 			{
-				UE_LOG(LogNeuralNetworkInference, Warning, TEXT("FImplBackEndUEOnly::Run(): Tensors.Num() = %d (should be > 0)."), TensorManager->GetTensorsMutable().Num());
+				UE_LOG(LogNeuralNetworkInference, Warning, TEXT("FImplBackEndUEOnly::Run(): Tensors.Num() = %d (should be > 0)."), TensorManager.GetTensorsMutable().Num());
 				return;
 			}
 
@@ -116,7 +111,7 @@ void UNeuralNetwork::FImplBackEndUEOnly::Run(FOnAsyncRunCompleted& InOutOnAsyncR
 					FRDGBuilder GraphBuilder(RHICmdList, RDG_EVENT_NAME("FImplBackEndUEOnly::Run()"));
 
 					// Move memory from CPU to GPU
-					TArray<FNeuralTensor>& Tensors = TensorManager->GetTensorsMutable();
+					TArray<FNeuralTensor>& Tensors = TensorManager.GetTensorsMutable();
 					const bool bIsInputInCPU = (InInputDeviceType == ENeuralDeviceType::CPU);
 					// Move only input tensors to GPU (once per FImplBackEndUEOnly::Run())
 					if (bAreTensorsInGpu)
@@ -124,11 +119,11 @@ void UNeuralNetwork::FImplBackEndUEOnly::Run(FOnAsyncRunCompleted& InOutOnAsyncR
 						// If input is on CPU, move inputs to GPU
 						if (bIsInputInCPU)
 						{
-							for (const int32 InputIndex : TensorManager->GetInputIndexes())
+							for (const int32 InputIndex : TensorManager.GetInputIndexes())
 							{
 								Tensors[InputIndex].ToGPU_RenderThread(&GraphBuilder);
 							}
-							for (const int32 NonInputIndex : TensorManager->GetNonInputIndexes())
+							for (const int32 NonInputIndex : TensorManager.GetNonInputIndexes())
 							{
 								Tensors[NonInputIndex].UpdateSRVAndOrUAV_RenderThread(&GraphBuilder);
 							}
@@ -158,11 +153,11 @@ void UNeuralNetwork::FImplBackEndUEOnly::Run(FOnAsyncRunCompleted& InOutOnAsyncR
 						// If input is already on GPU, just refresh inputs
 						else
 						{
-							for (const int32 InputIndex : TensorManager->GetInputIndexes())
+							for (const int32 InputIndex : TensorManager.GetInputIndexes())
 							{
 								Tensors[InputIndex].UpdateSRVAndOrUAV_RenderThread(&GraphBuilder);
 							}
-							for (const int32 NonInputIndex : TensorManager->GetNonInputIndexes())
+							for (const int32 NonInputIndex : TensorManager.GetNonInputIndexes())
 							{
 								Tensors[NonInputIndex].ToGPU_RenderThread(&GraphBuilder);
 							}
@@ -192,7 +187,7 @@ void UNeuralNetwork::FImplBackEndUEOnly::Run(FOnAsyncRunCompleted& InOutOnAsyncR
 					const bool bIsOutputInCPU = (InOutputDeviceType == ENeuralDeviceType::CPU);
 					if (bIsOutputInCPU)
 					{
-						for (const int32 OutputIndex : TensorManager->GetOutputIndexes())
+						for (const int32 OutputIndex : TensorManager.GetOutputIndexes())
 						{
 							Tensors[OutputIndex].ToCPU_RenderThread(&GraphBuilder);
 						}
@@ -238,8 +233,8 @@ void UNeuralNetwork::FImplBackEndUEOnly::Run(FOnAsyncRunCompleted& InOutOnAsyncR
 //	FString String = ModelProto.GetGraph().ToString();
 //	// Add FNeuralTensor(s)
 //	String += TEXT("TensorManager:\n");
-//	const TMap<FString, int32>& NameIndexMap = TensorManager->GetNameIndexMap();
-//	const TArray<FNeuralTensor>& Tensors = TensorManager->GetTensors();
+//	const TMap<FString, int32>& NameIndexMap = TensorManager.GetNameIndexMap();
+//	const TArray<FNeuralTensor>& Tensors = TensorManager.GetTensors();
 //	if (NameIndexMap.Num() > 0)
 //	{
 //		for (const auto& NameIndexPair : NameIndexMap)
@@ -255,12 +250,12 @@ void UNeuralNetwork::FImplBackEndUEOnly::Run(FOnAsyncRunCompleted& InOutOnAsyncR
 //		}
 //	}
 //	String += TEXT("InputTensorMap:\n");
-//	for (const auto& NameIndexPair : TensorManager->GetInputNameIndexMap())
+//	for (const auto& NameIndexPair : TensorManager.GetInputNameIndexMap())
 //	{
 //		String += FString::Format(TEXT(" -{0}: {1}\n"), { NameIndexPair.Key, Tensors[NameIndexPair.Value].ToString(20) });
 //	}
 //	String += TEXT("OutputTensorMap:\n");
-//	for (const auto& NameIndexPair : TensorManager->GetOutputNameIndexMap())
+//	for (const auto& NameIndexPair : TensorManager.GetOutputNameIndexMap())
 //	{
 //		String += FString::Format(TEXT(" -{0}: {1}\n"), { NameIndexPair.Key, Tensors[NameIndexPair.Value].ToString(20) });
 //	}
@@ -293,5 +288,4 @@ void UNeuralNetwork::FImplBackEndUEOnly::Reset(TSharedPtr<FImplBackEndUEOnly>& I
 		InOutImplBackEndUEOnly->ModelProto = FModelProto();
 		InOutImplBackEndUEOnly->bAreTensorsInGpu = false;
 	}
-	InOutImplBackEndUEOnly->TensorManager = MakeShared<FNeuralTensorManager>();
 }
diff --git a/Engine/Plugins/Experimental/NeuralNetworkInference/Source/NeuralNetworkInference/Private/NeuralNetworkImplBackEndUEOnly.h b/Engine/Plugins/Experimental/NeuralNetworkInference/Source/NeuralNetworkInference/Private/NeuralNetworkImplBackEndUEOnly.h
--- a/Engine/Plugins/Experimental/NeuralNetworkInference/Source/NeuralNetworkInference/Private/NeuralNetworkImplBackEndUEOnly.h
+++ b/Engine/Plugins/Experimental/NeuralNetworkInference/Source/NeuralNetworkInference/Private/NeuralNetworkImplBackEndUEOnly.h
@@ -6,7 +6,7 @@
 
 #include "ModelProto.h"
 #include "NeuralOperator.h"
-#include "NeuralTensorManager.h"
+#include "UEOnly/NeuralTensorManager.h"
 
 /* FImplBackEndUEOnly
  *****************************************************************************/
@@ -23,7 +23,7 @@ struct UNeuralNetwork::FImplBackEndUEOnly
 	/**
 	 * It contains a few TArray and TMaps for all FNeuralTensors (Input, Output, Intermediate(Not)Initialized, Weight).
 	 */
-	TSharedPtr<FNeuralTensorManager> TensorManager;
+	FNeuralTensorManager TensorManager;
 
 	/**
 	 * Only for the vanilla back end.
@@ -32,7 +32,7 @@ struct UNeuralNetwork::FImplBackEndUEOnly
 	TArray<TSharedPtr<FNeuralOperator>> Operators;
 
 	static bool Load(TSharedPtr<FImplBackEndUEOnly>& InOutImplBackEndUEOnly, const TArray<uint8>& InModelReadFromFileInBytes);
-	static bool Load(TSharedPtr<FImplBackEndUEOnly>& InOutImplBackEndUEOnly, TSharedPtr<FNeuralTensorManager>& InTensorManager, const TArray<TSharedPtr<FNeuralOperator>>& InOperators);
+	static bool Load(TSharedPtr<FImplBackEndUEOnly>& InOutImplBackEndUEOnly, FNeuralTensorManager& InTensorManager, const TArray<TSharedPtr<FNeuralOperator>>& InOperators);
 
 	//static bool Load(TSharedPtr<FImplBackEndUEOnly>& InOutImplBackEndUEOnly, FNeuralTensorManager& InTensorManager, const TArray<TSharedPtr<FNeuralOperator>>& InOperators);
 
diff --git a/Engine/Plugins/Experimental/NeuralNetworkInference/Source/NeuralNetworkInference/Private/NeuralTensorManager.cpp b/Engine/Plugins/Experimental/NeuralNetworkInference/Source/NeuralNetworkInference/Private/NeuralTensorManager.cpp
deleted file mode 100644
diff --git a/Engine/Plugins/Experimental/NeuralNetworkInference/Source/NeuralNetworkInference/Private/UEOnly/GraphProtoToNeuralNetworkConverter.cpp b/Engine/Plugins/Experimental/NeuralNetworkInference/Source/NeuralNetworkInference/Private/UEOnly/GraphProtoToNeuralNetworkConverter.cpp
new file mode 100644
--- /dev/null
+++ b/Engine/Plugins/Experimental/NeuralNetworkInference/Source/NeuralNetworkInference/Private/UEOnly/GraphProtoToNeuralNetworkConverter.cpp
@@ -0,0 +1,607 @@
+// Copyright Epic Games, Inc. All Rights Reserved.
+
+#include "GraphProtoToNeuralNetworkConverter.h"
+#include "Misc/Paths.h"
+#include "NeuralNetworkInferenceUtils.h"
+#include "NeuralOperator.h"
+#include "NeuralOperators.h"
+#include "NeuralTensorManager.h"
+
+
+
+/* FGraphProtoToNeuralNetworkConverter public functions
+ *****************************************************************************/
+
+bool FGraphProtoToNeuralNetworkConverter::Translate(TArray<TSharedPtr<FNeuralOperator>>& OutOperators, FNeuralTensorManager& InOrOutTensorManager, const FGraphProto& InGraphProto, const bool bInIsTensorManagerConst)
+{
+	// InOrOutTensorManager is const, i.e., InTensorManager
+	if (bInIsTensorManagerConst)
+	{
+		TMap<FString, int32> NameIndexMap = InOrOutTensorManager.GetNameIndexMap();
+		TMap<FString, int32> DummyOutputNameIndexMap;
+		const bool bResult = CreateOperatorsAndEditTensorArray(OutOperators, DummyOutputNameIndexMap, InOrOutTensorManager.GetTensorsMutable(), NameIndexMap, InOrOutTensorManager.GetInputNameIndexMap(),
+			InOrOutTensorManager.GetOutputNameIndexMap(), InGraphProto, bInIsTensorManagerConst);
+		if (NameIndexMap.Num() != InOrOutTensorManager.GetNameIndexMap().Num())
+		{
+			UE_LOG(LogNeuralNetworkInference, Warning, TEXT("FGraphProtoToNeuralNetworkConverter::Translate(): NameIndexMap changed its size, it must not (%d != %d)."),
+				NameIndexMap.Num(), InOrOutTensorManager.GetNameIndexMap().Num());
+		}
+		return bResult;
+	}
+	// InOrOutTensorManager is not const, i.e., OutTensorManager
+	else // if (!bInIsTensorManagerConst)
+	{
+		// MaxNumberTensorsInNetwork is required to avoid the bug of resizing and invalidating all pointers given to the operators
+		int32 MaxNumberTensorsInNetwork = 0;
+		{
+			// Quickly estimate the number of FNeuralTensor's of the UNeuralNetworkLegacy for Tensors.Reserve(...)
+			// Number FNeuralTensor's = Number (unique) input FNeuralTensors + number (unique) output FNeuralTensors
+			// Uniqueness guaranteed by using the TSet
+			TSet<FString> TensorNameSet;
+			for (const FNodeProto& NodeProto : InGraphProto.Node)
+			{
+				// Check input tensors
+				for (const FString& TensorName : NodeProto.Input)
+				{
+					if (TensorNameSet.Find(TensorName) == nullptr)
+					{
+						TensorNameSet.Add(TensorName);
+					}
+				}
+				// Check output tensors
+				for (const FString& TensorName : NodeProto.Output)
+				{
+					if (TensorNameSet.Find(TensorName) == nullptr)
+					{
+						TensorNameSet.Add(TensorName);
+					}
+				}
+			}
+			MaxNumberTensorsInNetwork = TensorNameSet.Num()
+			// Add maximum number possible of auxiliary tensors. A conservative estimate is the number of operators times the maximum possible number of auxiliary tensors per operator
+				+ InGraphProto.Node.Num() * FNeuralOperator::GetMaximumPossibleNumberAuxiliaryTensorsPerOperator();
+		}
+		TArray<FNeuralTensor> Tensors;
+		Tensors.Reserve(MaxNumberTensorsInNetwork);
+		TMap<FString, int32> NameIndexMap;
+
+		// Read input tensors
+		TMap<FString, int32> InputNameIndexMap;
+		for (const FValueInfoProto& GraphProtoInput : InGraphProto.Input)
+		{
+			// Create tensor
+			const TArray<FTensorShapeProtoDimension>& Dimensions = GraphProtoInput.Type.TensorType.Shape.Dim;
+			TArray<int64> DimensionsAsTArray;
+			for (const FTensorShapeProtoDimension& Dimension : Dimensions)
+			{
+				DimensionsAsTArray.Emplace(Dimension.DimValue);
+			}
+			// Add operator name to TMap
+			NameIndexMap.Add(GraphProtoInput.Name, Tensors.Num());
+			InputNameIndexMap.Add(GraphProtoInput.Name, Tensors.Num());
+			// Add tensor to TMap
+			Tensors.Emplace(FNeuralTensor(ENeuralDataType::Float, DimensionsAsTArray, GraphProtoInput.Name, ENeuralTensorTypeGPU::Input));
+		}
+
+		// Read output tensor names (but do not create them yet, it has to happen simultaneously with the operator creation and tensor inlining)
+		// Issue: For inlining to work properly, output tensors cannot be created until after inlining operators (i.e., CreateOperatorsAndEditTensorArray).
+		// Otherwise, the following failing case could happen with a network like: InputTensor --> Gemm --> IntermediateTensor --> Relu --> OutputTensor
+		// 1. Inputs and outputs are created before hand, i.e., InputTensor and OutputTensor
+		// 2. CreateOperatorsAndEditTensorArray() is run
+		//   2.1. Gemm and IntermediateTensor would be created (because Gemm cannot be inlined)
+		//   2.2. Relu would be created, Relu would verify its input can be inlined, but both IntermediateTensor and OutputTensor are already created at this point, so it cannot inline them (crash/bug)
+		TMap<FString, int32> OutputNameDummyIndexMap;
+		for (const FValueInfoProto& GraphProtoOutput : InGraphProto.Output)
+		{
+			// Add operator name to TMap
+			if (NameIndexMap.Find(GraphProtoOutput.Name))
+			{
+				UE_LOG(LogNeuralNetworkInference, Warning, TEXT("FGraphProtoToNeuralNetworkConverter::Translate(): A FNeuralTensor of the UNeuralNetworkLegacy is a global Input and Output simultaneously."));
+			}
+			OutputNameDummyIndexMap.Add(GraphProtoOutput.Name, -1);
+		}
+
+		// Fill operators
+		TMap<FString, int32> OutputNameIndexMap;
+		bool bWasSuccessful = CreateOperatorsAndEditTensorArray(OutOperators, OutputNameIndexMap, Tensors, NameIndexMap, InputNameIndexMap, OutputNameDummyIndexMap, InGraphProto, bInIsTensorManagerConst);
+		// Sanity check
+		if (Tensors.Num() > MaxNumberTensorsInNetwork)
+		{
+			UE_LOG(LogNeuralNetworkInference, Warning, TEXT("FGraphProtoToNeuralNetworkConverter::Translate(): Tensors.Num() <= MaxNumberTensorsInNetwork failed (%d vs. %d). A resize of Tensors"
+				" might have occurred, invalidating all TArray<FNeuralTensor*> FOperator variables. Report this error (so we can fix MaxNumberTensorsInNetwork to be big enough)."),
+				Tensors.Num(), MaxNumberTensorsInNetwork);
+			bWasSuccessful = false;
+		}
+
+		if (bWasSuccessful)
+		{
+			if (Tensors.Num() <= NameIndexMap.Num())
+			{
+				UE_LOG(LogNeuralNetworkInference, Display, TEXT("Auto-inlining: %d initial tensors compressed into %d (%d tensors merged)."),
+					NameIndexMap.Num(), Tensors.Num(), NameIndexMap.Num() - Tensors.Num());
+			}
+			// Set InOrOutTensorManager
+			return InOrOutTensorManager.Load(Tensors, NameIndexMap, InputNameIndexMap, OutputNameIndexMap);
+		}
+		return bWasSuccessful;
+	}
+}
+
+
+
+/* FGraphProtoToNeuralNetworkConverter private functions
+ *****************************************************************************/
+
+bool FGraphProtoToNeuralNetworkConverter::CreateOperatorsAndEditTensorArray(TArray<TSharedPtr<FNeuralOperator>>& OutOperators, TMap<FString, int32>& OutputNameIndexMap,
+	TArray<FNeuralTensor>& InOutTensors, TMap<FString, int32>& InOutNameIndexMap, const TMap<FString, int32>& InInputNameDummyIndexMap,
+	const TMap<FString, int32>& InOutputNameDummyIndexMap, const FGraphProto& InGraphProto, const bool bInIsTensorManagerConst)
+{
+	const int64 InOutTensorsNumInit = InOutTensors.Num();
+	const int64 InOutNameIndexMapNumInit = InOutNameIndexMap.Num();
+	// Create operators sequentially
+	OutOperators.Empty();
+	OutputNameIndexMap.Empty();
+	// Iterating over each operator
+	for (const FNodeProto& NodeProto : InGraphProto.Node)
+	{
+		// Read/set input tensors
+		TArray<FNeuralTensor*> OperatorInputTensors;
+		TArray<FString> OperatorInputTensorNames;
+		// Iterating over each input tensor name for that operator
+		for (const FString& InputTensorName : NodeProto.Input)
+		{
+			OperatorInputTensorNames.Push(InputTensorName);
+			// FNeuralTensor already read (e.g., absolute input of the UNeuralNetworkLegacy, or the output of a previous tensor required as input of this one)
+			if (int32* ExistingTensorIndex = InOutNameIndexMap.Find(InputTensorName))
+			{
+				OperatorInputTensors.Push(&InOutTensors[*ExistingTensorIndex]);
+			}
+			// FNeuralTensor was never modified, so it must come from a weight tensor (regardless of whether it is also an absolute output of the network)
+			else
+			{
+				const FTensorProto& TensorProto = *FModelProto::FindElementInArray(InputTensorName, InGraphProto.Initializer, /*bMustValueBeFound*/true);
+				// Is it also an absolute output of the network?
+				const int64 TensorIndex = InOutTensors.Num();
+				if (!bInIsTensorManagerConst && InOutputNameDummyIndexMap.Find(InputTensorName))
+				{
+					OutputNameIndexMap.Add(InputTensorName, TensorIndex);
+				}
+				// Create and add tensor to InOutTensors
+				InOutTensors.Push(FNeuralTensor());
+				FNeuralTensor& NewTensor = InOutTensors.Last();
+				// Deprecated code for otxt files, if added back, add before the 2 for loops for speedup
+				// FString PathPart, FilenamePart, ExtensionPart;
+				// if (!bInIsTensorManagerConst)
+				// {
+				// 	FPaths::Split(InModelPath, PathPart, FilenamePart, ExtensionPart);
+				// }
+				// Get NewTensor from TensorProto
+				if (!NewTensor.SetFromTensorProto(&TensorProto, InputTensorName, ENeuralTensorTypeGPU::Weight/*, PathPart*/))
+				{
+					UE_LOG(LogNeuralNetworkInference, Warning, TEXT("FGraphProtoToNeuralNetworkConverter::CreateOperatorsAndEditTensorArray(): SetFromTensorProto() failed."));
+					return false;
+				}
+				// Update InOutNameIndexMap and OperatorInputTensors
+				InOutNameIndexMap.Add(InputTensorName, TensorIndex);
+				OperatorInputTensors.Push(&NewTensor);
+			}
+		}
+		// Nodes that can be replaced for inlined functions and do not need to be executed
+		if (NodeProto.OperatorType == TEXT("Constant"))
+		{
+			// Sanity check
+			if (NodeProto.Output.Num() != 1)
+			{
+				UE_LOG(LogNeuralNetworkInference, Warning, TEXT("FGraphProtoToNeuralNetworkConverter::CreateOperatorsAndEditTensorArray(): Constant operator requires exactly 1 output, not %d."), NodeProto.Output.Num());
+				return false;
+			}
+			const FString& OutputTensorName = NodeProto.Output[0];
+			// Case 1: bInIsTensorManagerConst
+			if (bInIsTensorManagerConst)
+			{
+				// Sanity check
+				if (!InOutNameIndexMap.Find(OutputTensorName))
+				{
+					UE_LOG(LogNeuralNetworkInference, Warning, TEXT("FGraphProtoToNeuralNetworkConverter::CreateOperatorsAndEditTensorArray()-Constant: Tensor %s was not found."), *OutputTensorName);
+					return false;
+				}
+				// Do nothing else
+			}
+			// Case 2: !bInIsTensorManagerConst
+			else
+			{
+				// Sanity check
+				if (InOutNameIndexMap.Find(OutputTensorName))
+				{
+					UE_LOG(LogNeuralNetworkInference, Warning, TEXT("FGraphProtoToNeuralNetworkConverter::CreateOperatorsAndEditTensorArray()-Constant: Tensor %s was already defined and it should have not."), *OutputTensorName);
+					return false;
+				}
+				// Is it also an absolute output of the network?
+				const int64 TensorIndex = InOutTensors.Num();
+				if (InOutputNameDummyIndexMap.Find(OutputTensorName))
+				{
+					OutputNameIndexMap.Add(OutputTensorName, TensorIndex);
+				}
+				// Create tensor
+				InOutNameIndexMap.Add(OutputTensorName, TensorIndex);
+				InOutTensors.Push(FNeuralTensor(ENeuralDataType::None, TArray<int64>({}), OutputTensorName, ENeuralTensorTypeGPU::Weight));
+
+				// Create constant operator layer and get final tensor value
+				FConstantOperator ConstantOperator(NodeProto);
+				ConstantOperator.SetOutputTensors({&InOutTensors.Last()});
+				// Sanity check
+				if (!ConstantOperator.ConfigureOutputAndInternalVariablesAndSanityChecks())
+				{
+					UE_LOG(LogNeuralNetworkInference, Warning, TEXT("FGraphProtoToNeuralNetworkConverter::CreateOperatorsAndEditTensorArray()-Constant: Tensor %s was already defined and it should have not."), *OutputTensorName);
+					return false;
+				}
+				// Get output
+				ConstantOperator.ForwardCPU();
+			}
+		}
+		// Nodes/Operators that must be created and used on each Run() pass of the network
+		else
+		{
+			// Create and configure operator
+			TSharedPtr<FNeuralOperator> Operator = CreateOperator(NodeProto, OperatorInputTensorNames, InInputNameDummyIndexMap, InOutputNameDummyIndexMap, InGraphProto);
+			if (!Operator.IsValid())
+			{
+				return false;
+			}
+			// ConvTranspose requires flipping OperatorInputTensors[1], so we expect for now a Weight type, so it can be flipped beforehand
+			if (!bInIsTensorManagerConst && NodeProto.OperatorType == TEXT("ConvTranspose")) //if (NodeProto.OperatorType.Contains(TEXT("Conv"), ESearchCase::CaseSensitive))
+			{
+				UE_LOG(LogNeuralNetworkInference, Display, TEXT("Flipping FConvTranspose weights. This is a temporary measurement but it will keep the weights flipped for now."));
+				// Sanity check
+				if (OperatorInputTensors.Num() > 1 && OperatorInputTensors[1]->GetTensorTypeGPU() != ENeuralTensorTypeGPU::Weight)
+				{
+					UE_LOG(LogNeuralNetworkInference, Warning,
+						TEXT("FGraphProtoToNeuralNetworkConverter::CreateOperatorsAndEditTensorArray(): For now, %s operators require a constant InputTensor[0] (i.e., a fixed weight tensor), it cannot change."), *NodeProto.OperatorType);
+					return false;
+				}
+				// Flip (invert) OperatorInputTensors[1]
+				if (!OperatorInputTensors[1]->Flip(2, OperatorInputTensors[1]->GetNumberDimensions()))
+				{
+					UE_LOG(LogNeuralNetworkInference, Warning, TEXT("FGraphProtoToNeuralNetworkConverter::CreateOperatorsAndEditTensorArray(): Tensor::Flip() failed."));
+					return false;
+				}
+			}
+			// Set input tensors in Operator
+			const int32 InlinedTensorIndex = Operator->SetInputTensorsAndGetInlinedIndex(OperatorInputTensors);
+			// Read/set output tensors
+			// 2 different cases for OperatorOutputTensors:
+				// Case 1/2: bInIsTensorManagerConst --> All FNeuralTensors already exist
+					// Subcase 1/2: InlinedTensorIndex > -1 --> Do nothing
+					// Subcase 2/2: InlinedTensorIndex < 0 --> Fill OperatorOutputTensors from known index
+				// Case 2/2: !bInIsTensorManagerConst
+					// Subcase 1/3: ExistingTensorIndex is found --> FNeuralTensor already exists
+					// Subcase 2/3: InlinedTensorIndex < 0 --> Output tensor must be created because operator creates its own output
+					// Subcase 3/3: InlinedTensorIndex > -1 --> Output tensor is the same than 1 of the input tensors of that operator (inlined operator)
+			if (InlinedTensorIndex > -1 && NodeProto.Output.Num() != 1)
+			{
+				UE_LOG(LogNeuralNetworkInference, Warning, TEXT("FGraphProtoToNeuralNetworkConverter::CreateOperatorsAndEditTensorArray(): There cannot be an inlined tensor (index %d)"
+					" if NodeProto.Output.Num() > 1 (currently %d)."), InlinedTensorIndex, NodeProto.Output.Num());
+				return false;
+			}
+			TArray<FNeuralTensor*> OperatorOutputTensors;
+			for (const FString& OutputTensorName : NodeProto.Output)
+			{
+				// Case 1: bInIsTensorManagerConst
+				if (bInIsTensorManagerConst)
+				{
+					// Subcase 1/2: Do nothing
+					// Subcase 2/2: Fill OperatorOutputTensors from known index
+					if (InlinedTensorIndex < 0)
+					{
+						const int32 ExistingTensorIndex = InOutNameIndexMap.FindChecked(OutputTensorName);
+						OperatorOutputTensors.Push(&InOutTensors[ExistingTensorIndex]);
+					}
+				}
+				// Case 2: !bInIsTensorManagerConst
+				else
+				{
+					// Subcase 1/3: FNeuralTensor already exists (i.e., absolute output of the UNeuralNetworkLegacy or bInIsTensorManagerConst == true)
+					if (const int32* const ExistingTensorIndex = InOutNameIndexMap.Find(OutputTensorName))
+					{
+						if (InlinedTensorIndex > -1)
+						{
+							UE_LOG(LogNeuralNetworkInference, Warning,
+								TEXT("FGraphProtoToNeuralNetworkConverter::CreateOperatorsAndEditTensorArray(): InlinedTensorIndex = %d for this operator (%s) but the tensor (%s) already exists."), InlinedTensorIndex, *NodeProto.OperatorType, *OutputTensorName);
+							return false;
+						}
+						OperatorOutputTensors.Push(&InOutTensors[*ExistingTensorIndex]);
+						// Is it also an absolute output of the network?
+						if (InOutputNameDummyIndexMap.Find(OutputTensorName))
+						{
+							OutputNameIndexMap.Add(OutputTensorName, *ExistingTensorIndex);
+							InOutTensors[*ExistingTensorIndex].SetTensorTypeGPU(ENeuralTensorTypeGPU::Output);
+						}
+					}
+					// Subcases 2 and 3: FNeuralTensor not found, it must be defined now (i.e., intermediate tensor of the UNeuralNetworkLegacy)
+					// Subcase 2/3: Output tensor must be created because operator creates its own output
+					else if (InlinedTensorIndex < 0)
+					{
+						// Is it also an absolute output of the network?
+						const int64 TensorIndex = InOutTensors.Num();
+						ENeuralTensorTypeGPU TensorTypeGPU = ENeuralTensorTypeGPU::IntermediateNotInitialized;
+						if (InOutputNameDummyIndexMap.Find(OutputTensorName))
+						{
+							OutputNameIndexMap.Add(OutputTensorName, TensorIndex);
+							TensorTypeGPU = ENeuralTensorTypeGPU::Output;
+						}
+						// Create tensor
+						InOutNameIndexMap.Add(OutputTensorName, TensorIndex);
+						InOutTensors.Push(FNeuralTensor(ENeuralDataType::Float, TArray<int64>({}), OutputTensorName, TensorTypeGPU));
+						OperatorOutputTensors.Push(&InOutTensors.Last());
+					}
+					// Subcase 3/3: Output tensor is the same than 1 of the input tensors of that operator (inlined operator)
+					else
+					{
+						if (NodeProto.Output.Num() != 1)
+						{
+							UE_LOG(LogNeuralNetworkInference, Warning, TEXT("FGraphProtoToNeuralNetworkConverter::CreateOperatorsAndEditTensorArray()-%s-%s:"
+								" NodeProto.Output.Num() = %d != 1. Only implemented for inlined operators with 1 output. Potential issues for operators like PReLU."
+								" Although it should work as long as we pre-define the desired order of FNeuralTensors inside that operator."),
+								*NodeProto.OperatorType, *Operator->GetName(), OperatorInputTensors.Num(), NodeProto.Output.Num());
+							return false;
+						}
+						const FString& NodeProtoInputName = NodeProto.Input[InlinedTensorIndex];
+						const int32 InlinedTensorGlobalIndex = InOutNameIndexMap.FindChecked(NodeProtoInputName);
+						InOutNameIndexMap.Add(OutputTensorName, InlinedTensorGlobalIndex);
+						// Is it also an absolute output of the network?
+						if (InOutputNameDummyIndexMap.Find(OutputTensorName))
+						{
+							OutputNameIndexMap.Add(OutputTensorName, InlinedTensorGlobalIndex);
+							InOutTensors[InlinedTensorGlobalIndex].SetTensorTypeGPU(ENeuralTensorTypeGPU::Output);
+						}
+					}
+				}
+			}
+			// Set output tensors in Operator (if operator is not fully inlined)
+			if (OperatorOutputTensors.Num() > 0)
+			{
+				Operator->SetOutputTensors(OperatorOutputTensors);
+			}
+			// Create and add auxiliary tensors for that operator (if any)
+			const int64 NumberAuxiliaryTensors = Operator->GetNumberAuxiliaryTensors();
+			if (NumberAuxiliaryTensors > 0)
+			{
+				TArray<FNeuralTensor*> AuxiliaryOperatorTensors;
+				for (int64 AuxiliaryTensorIndex = 0; AuxiliaryTensorIndex < Operator->GetNumberAuxiliaryTensors(); ++AuxiliaryTensorIndex)
+				{
+					const FString AuxiliaryTensorName = NodeProto.Name + TEXT("-") + Operator->GetName() + TEXT("-Aux") + FString::FromInt(AuxiliaryTensorIndex);
+					// Uasset file
+					if (bInIsTensorManagerConst)
+					{
+						// Find and push auxiliary tensor
+						if (const int32* const AuxiliaryTensorAbsoluteIndex = InOutNameIndexMap.Find(AuxiliaryTensorName))
+						{
+							AuxiliaryOperatorTensors.Push(&InOutTensors[*AuxiliaryTensorAbsoluteIndex]);
+						}
+						// Something went wrong, the tensor should exist
+						else
+						{
+							UE_LOG(LogNeuralNetworkInference, Warning, TEXT("FGraphProtoToNeuralNetworkConverter::CreateOperatorsAndEditTensorArray(): AuxiliaryTensorName %s did not exist in InOutTensors."), *AuxiliaryTensorName);
+							return false;
+						}
+					}
+					// ONNX file
+					else
+					{
+						// Sanity check
+						if (InOutNameIndexMap.Find(AuxiliaryTensorName))
+						{
+							UE_LOG(LogNeuralNetworkInference, Warning, TEXT("FGraphProtoToNeuralNetworkConverter::CreateOperatorsAndEditTensorArray(): InOutNameIndexMap already has %s. Notify us of this error."), *AuxiliaryTensorName);
+							return false;
+						}
+						// Create and push new auxiliary tensor
+						InOutNameIndexMap.Add(AuxiliaryTensorName, InOutTensors.Num());
+						InOutTensors.Push(FNeuralTensor(ENeuralDataType::None, TArray<int64>({}), AuxiliaryTensorName, ENeuralTensorTypeGPU::Generic)); // ENeuralTensorTypeGPU can be modify by FNeuralOperator::Configure()
+						AuxiliaryOperatorTensors.Push(&InOutTensors.Last());
+					}
+				}
+				Operator->SetAuxiliaryTensors(AuxiliaryOperatorTensors);
+			}
+			// If operator is inlined, then input and output should already match
+			// Configure operator
+			if (!Operator->Configure())
+			{
+				UE_LOG(LogNeuralNetworkInference, Warning,
+					TEXT("FGraphProtoToNeuralNetworkConverter::CreateOperatorsAndEditTensorArray(): Operator %s could not be configured, so network could not be loaded."), *Operator->GetName());
+				return false;
+			}
+			// Push Operator into OutOperators
+			OutOperators.Push(Operator);
+		}
+	}
+	// Sanity checks for bInIsTensorManagerConst
+	if (bInIsTensorManagerConst)
+	{
+		if (OutputNameIndexMap.Num() > 0)
+		{
+			UE_LOG(LogNeuralNetworkInference, Warning,
+				TEXT("FGraphProtoToNeuralNetworkConverter::CreateOperatorsAndEditTensorArray(): Error in this function, OutputNameIndexMap.Num() should be 0 and not %d when reading from serialized data."),
+				OutputNameIndexMap.Num());
+			return false;
+		}
+		else if (InOutTensorsNumInit != InOutTensors.Num())
+		{
+			UE_LOG(LogNeuralNetworkInference, Warning,
+				TEXT("FGraphProtoToNeuralNetworkConverter::CreateOperatorsAndEditTensorArray(): Error in this function, InOutTensors.Num() should not change from %d to %d when reading from serialized data."),
+				InOutTensorsNumInit, InOutTensors.Num());
+			return false;
+		}
+		else if (InOutNameIndexMapNumInit != InOutNameIndexMap.Num())
+		{
+			UE_LOG(LogNeuralNetworkInference, Warning,
+				TEXT("FGraphProtoToNeuralNetworkConverter::CreateOperatorsAndEditTensorArray(): Error in this function, InOutNameIndexMap.Num() should not change from %d to %d when reading from serialized data."),
+				InOutTensorsNumInit, InOutTensors.Num());
+			return false;
+		}
+	}
+	// Sanity check for !bInIsTensorManagerConst
+	else if (OutputNameIndexMap.Num() != InOutputNameDummyIndexMap.Num())
+	{
+		UE_LOG(LogNeuralNetworkInference, Warning,
+			TEXT("FGraphProtoToNeuralNetworkConverter::CreateOperatorsAndEditTensorArray(): OutputNameIndexMap.Num() == InOutputNameDummyIndexMap.Num() failed, %d != %d."
+				" This means that not all of the outputs defined on the ONNX file were actually needed. The program will proceed with only the required ones."),
+			OutputNameIndexMap.Num(), InOutputNameDummyIndexMap.Num());
+	}
+	// Sanity check for all cases
+	if (InOutNameIndexMap.Num() < InOutTensors.Num())
+	{
+		UE_LOG(LogNeuralNetworkInference, Warning,
+			TEXT("FGraphProtoToNeuralNetworkConverter::CreateOperatorsAndEditTensorArray(): Error in this function, InOutNameIndexMap.Num() >= InOutTensors.Num() failed, %d != %d. They"
+				" should match if no inlining occurs, or more generically, InOutNameIndexMap.Num() = InOutTensors.Num() + number_inlined_tensors."),
+			InOutNameIndexMap.Num(), InOutTensors.Num());
+	}
+
+	return (OutOperators.Num() > 0);
+}
+
+bool FGraphProtoToNeuralNetworkConverter::CanTensorBeInlined(const FString& InTensorName, const TMap<FString, int32>& InInputNameDummyIndexMap,
+	const TMap<FString, int32>& InOutputNameDummyIndexMap, const FNodeProto& InNodeProto, const FGraphProto& InGraphProto, const bool bCanOutputPropertiesChangeWithPostForward)
+{
+	// If tensor is one of the input/output tensors --> Operator cannot be inlined
+	if (InInputNameDummyIndexMap.Find(InTensorName) || InOutputNameDummyIndexMap.Find(InTensorName))
+	{
+		return false;
+	}
+	// If PostForward() changes the properties of the output tensor (e.g., Reshape might change its size), then it cannot be combined with a global output of the network
+	else if (bCanOutputPropertiesChangeWithPostForward)
+	{
+		for (const FValueInfoProto& GraphProtoOutput : InGraphProto.Output)
+		{
+			if (GraphProtoOutput.Name == InNodeProto.Output[0])
+			{
+				return false;
+			}
+		}
+	}
+	// If tensor is used again --> Operator cannot be inlined
+	else
+	{
+		int32 NumberOperatorsUsingTensor = 0;
+		// Iterating over each operator
+		for (const FNodeProto& NodeProto : InGraphProto.Node)
+		{
+			// Iterating over each input tensor name for that operator
+			for (const FString& InputTensorName : NodeProto.Input)
+			{
+				// If names matches, increase NumberOperatorsUsingTensor
+				if (InputTensorName == InTensorName)
+				{
+					++NumberOperatorsUsingTensor;
+					if (NumberOperatorsUsingTensor > 1)
+					{
+						return false;
+					}
+				}
+			}
+		}
+		// Sanity check
+		if (NumberOperatorsUsingTensor != 1)
+		{
+			UE_LOG(LogNeuralNetworkInference, Warning,
+				TEXT("FGraphProtoToNeuralNetworkConverter::CanTensorBeInlined(): At this point, NumberOperatorsUsingTensor should be exactly 1 (not %d)."),
+				NumberOperatorsUsingTensor);
+		}
+	}
+	// None of the above is true --> Operator can be safely inlined
+	return true;
+}
+
+TSet<uint32> FGraphProtoToNeuralNetworkConverter::GetPotentiallyInlinedTensors(const TArray<FString>& InTensorNames, const TMap<FString, int32>& InInputNameDummyIndexMap,
+	const TMap<FString, int32>& InOutputNameDummyIndexMap, const FNodeProto& InNodeProto, const FGraphProto& InGraphProto)
+{
+	TSet<uint32> PotentiallyInlinedTensors;
+	for (int32 TensorNameIndex = 0; TensorNameIndex < InTensorNames.Num(); ++TensorNameIndex)
+	{
+		const FString& TensorName = InTensorNames[TensorNameIndex];
+		if (FGraphProtoToNeuralNetworkConverter::CanTensorBeInlined(TensorName, InInputNameDummyIndexMap, InOutputNameDummyIndexMap, InNodeProto, InGraphProto))
+		{
+			PotentiallyInlinedTensors.Add((uint32)TensorNameIndex);
+		}
+	}
+	return PotentiallyInlinedTensors;
+}
+
+// Eg Element-wise operators with 0 attributes
+#define IF_MATCHED_INLINABLE_OPERATOR_INITIALIZATION(OperatorTypeSuffix) \
+	if (InNodeProto.OperatorType == TEXT(#OperatorTypeSuffix)) /* Eg TEXT("Relu") */ \
+	{ \
+		const bool bCanTensorBeInlined = CanTensorBeInlined(InInputTensorNamesForOperator[0], InInputNameDummyIndexMap, InOutputNameDummyIndexMap, InNodeProto, InGraphProto); \
+		Operator = MakeShared<F##OperatorTypeSuffix##Operator>(bCanTensorBeInlined); /* Eg FReluOperator */ \
+	}
+
+// Eg Element-wise operators with n attributes or operators like BatchNormalization
+#define IF_MATCHED_INLINABLE_OPERATOR_WITH_PROTO_INITIALIZATION(OperatorTypeSuffix, bCanOutputPropertiesChangeWithPostForward) \
+	if (InNodeProto.OperatorType == TEXT(#OperatorTypeSuffix)) /* Eg TEXT("BatchNormalization") */ \
+	{ \
+		const bool bCanTensorBeInlined = CanTensorBeInlined(InInputTensorNamesForOperator[0], InInputNameDummyIndexMap, InOutputNameDummyIndexMap, InNodeProto, InGraphProto, bCanOutputPropertiesChangeWithPostForward); \
+		Operator = MakeShared<F##OperatorTypeSuffix##Operator>(bCanTensorBeInlined, InNodeProto); /* Eg FBatchNormalizationOperator */ \
+	}
+
+// Eg Conv(Transpose) and Gemm
+#define IF_MATCHED_OPERATOR_WITH_PROTO_INITIALIZATION(OperatorTypeSuffix) \
+	if (InNodeProto.OperatorType == TEXT(#OperatorTypeSuffix)) /* Eg TEXT("Conv") */ \
+	{ \
+		Operator = MakeShared<F##OperatorTypeSuffix##Operator>(InNodeProto); /* Eg FConvOperator */ \
+	}
+
+#define IF_MATCHED_MULTIDIRECTIONAL_BROADCASTING_OPERATOR_INITIALIZATION(OperatorTypeSuffix) \
+	if (InNodeProto.OperatorType == TEXT(#OperatorTypeSuffix)) /* Eg TEXT("Add") */ \
+	{ \
+		const TSet<uint32> PotentiallyInlinedTensors = GetPotentiallyInlinedTensors(InInputTensorNamesForOperator, InInputNameDummyIndexMap, InOutputNameDummyIndexMap, InNodeProto, InGraphProto); \
+		Operator = MakeShared<F##OperatorTypeSuffix##Operator>(PotentiallyInlinedTensors); /* Eg FAddOperator */ \
+	}
+
+TSharedPtr<FNeuralOperator> FGraphProtoToNeuralNetworkConverter::CreateOperator(const FNodeProto& InNodeProto, const TArray<FString>& InInputTensorNamesForOperator,
+	const TMap<FString, int32>& InInputNameDummyIndexMap, const TMap<FString, int32>& InOutputNameDummyIndexMap, const FGraphProto& InGraphProto)
+{
+	TSharedPtr<FNeuralOperator> Operator;
+	IF_MATCHED_INLINABLE_OPERATOR_INITIALIZATION(Abs)
+	else IF_MATCHED_INLINABLE_OPERATOR_INITIALIZATION(Acos)
+	else IF_MATCHED_MULTIDIRECTIONAL_BROADCASTING_OPERATOR_INITIALIZATION(Add)
+	else IF_MATCHED_INLINABLE_OPERATOR_INITIALIZATION(Asin)
+	else IF_MATCHED_INLINABLE_OPERATOR_INITIALIZATION(Atan)
+	else IF_MATCHED_INLINABLE_OPERATOR_WITH_PROTO_INITIALIZATION(BatchNormalization, false)
+	else IF_MATCHED_INLINABLE_OPERATOR_INITIALIZATION(Ceil)
+	else IF_MATCHED_OPERATOR_WITH_PROTO_INITIALIZATION(Conv)
+	else IF_MATCHED_OPERATOR_WITH_PROTO_INITIALIZATION(ConvTranspose)
+	if (InNodeProto.OperatorType == TEXT("ConvTranspose"))
+	{
+		TSharedPtr<FConvTransposeOperator> ConvTransposeOperator;
+		ConvTransposeOperator = MakeShared<FConvTransposeOperator>(InNodeProto);
+		ConvTransposeOperator->SetWhetherWeightTensorWasFlipped(true);
+		Operator = MoveTemp(ConvTransposeOperator);
+	}
+	else IF_MATCHED_INLINABLE_OPERATOR_INITIALIZATION(Cos)
+	else IF_MATCHED_INLINABLE_OPERATOR_INITIALIZATION(Cosh)
+	else IF_MATCHED_MULTIDIRECTIONAL_BROADCASTING_OPERATOR_INITIALIZATION(Div)
+	else IF_MATCHED_INLINABLE_OPERATOR_INITIALIZATION(Exp)
+	else IF_MATCHED_INLINABLE_OPERATOR_INITIALIZATION(Floor)
+	else IF_MATCHED_OPERATOR_WITH_PROTO_INITIALIZATION(Gemm)
+	else IF_MATCHED_INLINABLE_OPERATOR_WITH_PROTO_INITIALIZATION(LeakyRelu, false)
+	else IF_MATCHED_INLINABLE_OPERATOR_INITIALIZATION(Log)
+	else IF_MATCHED_MULTIDIRECTIONAL_BROADCASTING_OPERATOR_INITIALIZATION(Mul)
+	else IF_MATCHED_INLINABLE_OPERATOR_INITIALIZATION(Neg)
+	else IF_MATCHED_MULTIDIRECTIONAL_BROADCASTING_OPERATOR_INITIALIZATION(Pow)
+	else IF_MATCHED_INLINABLE_OPERATOR_INITIALIZATION(Reciprocal)
+	else IF_MATCHED_INLINABLE_OPERATOR_INITIALIZATION(Relu)
+	else IF_MATCHED_INLINABLE_OPERATOR_WITH_PROTO_INITIALIZATION(Reshape, true)
+	else IF_MATCHED_INLINABLE_OPERATOR_INITIALIZATION(Round)
+	else IF_MATCHED_INLINABLE_OPERATOR_INITIALIZATION(Sigmoid)
+	else IF_MATCHED_INLINABLE_OPERATOR_INITIALIZATION(Sign)
+	else IF_MATCHED_INLINABLE_OPERATOR_INITIALIZATION(Sin)
+	else IF_MATCHED_INLINABLE_OPERATOR_INITIALIZATION(Sinh)
+	else IF_MATCHED_INLINABLE_OPERATOR_INITIALIZATION(Sqrt)
+	else IF_MATCHED_MULTIDIRECTIONAL_BROADCASTING_OPERATOR_INITIALIZATION(Sub)
+	else IF_MATCHED_INLINABLE_OPERATOR_INITIALIZATION(Tan)
+	else IF_MATCHED_INLINABLE_OPERATOR_INITIALIZATION(Tanh)
+	if (!Operator.IsValid())
+	{
+		UE_LOG(LogNeuralNetworkInference, Warning, TEXT("FGraphProtoToNeuralNetworkConverter::CreateOperator(): Unknown or unimplemented operator %s."), *InNodeProto.OperatorType);
+	}
+	return Operator;
+}
diff --git a/Engine/Plugins/Experimental/NeuralNetworkInference/Source/NeuralNetworkInference/Private/UEOnly/GraphProtoToNeuralNetworkConverter.h b/Engine/Plugins/Experimental/NeuralNetworkInference/Source/NeuralNetworkInference/Private/UEOnly/GraphProtoToNeuralNetworkConverter.h
new file mode 100644
--- /dev/null
+++ b/Engine/Plugins/Experimental/NeuralNetworkInference/Source/NeuralNetworkInference/Private/UEOnly/GraphProtoToNeuralNetworkConverter.h
@@ -0,0 +1,62 @@
+// Copyright Epic Games, Inc. All Rights Reserved.
+
+#pragma once
+
+#include "CoreMinimal.h"
+#include "ModelProto.h"
+#include "NeuralOperator.h"
+#include "NeuralTensor.h"
+#include "NeuralTensorManager.h"
+
+class FGraphProtoToNeuralNetworkConverter
+{
+public:
+	/**
+	 * It takes a FGraphProto and model path as inputs, and produces an array of FNeuralOperator's and FNeuralTensorManager as result.
+	 * These 2 outputs are used by UNeuralNetworkLegacy to create and load a (deep) neural network architecture and its weights.
+	 * @param InTensorManager It is const, but kept not const for CreateOperatorsAndEditTensorArray() compatibility.
+	 *
+	 * 2 alternatives:
+	 * If bInIsTensorManagerConst == true, InOrOutTensorManager is assumed to be a const input variable (assuming the tensors are all already loaded, e.g., on de-serialization).
+	 * With InModelPath, InOrOutTensorManager is assumed to be an output variable that will be filled based on InModelPath.
+	 */
+	static bool Translate(TArray<TSharedPtr<FNeuralOperator>>& OutOperators, FNeuralTensorManager& InOrOutTensorManager, const FGraphProto& InGraphProto, const bool bInIsTensorManagerConst);
+
+private:
+	/**
+	 * Auxiliary function for Translate().
+	 * It initializes both OutOperators and OutputNameIndexMap, as well as modify InOutTensors and InOutNameIndexMap.
+	 * OutputNameIndexMap is just the copy of InOutputNameDummyIndexMap while filling its int32 Value fields. Given the dual Game/Editor behaviors, and to optimize Game mode, kept as 2 different rather than modifying in-place as InOutOutputNameDummyIndexMap.
+	 *
+	 * 2 alternatives:
+	 * - If bInIsTensorManagerConst == true, InOrOutTensorManager is assumed to be a const input variable (assuming the tensors are all already loaded, e.g., on de-serialization).
+	 * - If bInIsTensorManagerConst == false, InOrOutTensorManager is assumed to be an output variable that will be filled based on InModelPath.
+	 *
+	 * @param OutputNameIndexMap It is created from InOutputNameDummyIndexMap, assuming the FString part of the mapping is right, and it will fill the index (int32 part of it).
+	 * @param InInputNameDummyIndexMap The int32 part is irrelevant, the only important part is whether the FString field is filled or not, regardless of its value. Kept as TMap for performance to integrate better with FTensorManager.
+	 * @param InOutputNameDummyIndexMap The int32 part is irrelevant, the only important part is whether the FString field is filled or not, regardless of its value. Kept as TMap for performance to integrate better with FTensorManager.
+	 */
+	static bool CreateOperatorsAndEditTensorArray(TArray<TSharedPtr<FNeuralOperator>>& OutOperators, TMap<FString, int32>& OutputNameIndexMap, TArray<FNeuralTensor>& InOutTensors,
+		TMap<FString, int32>& InOutNameIndexMap, const TMap<FString, int32>& InInputNameDummyIndexMap, const TMap<FString, int32>& InOutputNameDummyIndexMap,
+		const FGraphProto& InGraphProto, const bool bInIsTensorManagerConst);
+
+	/**
+	 * It will return FNeuralOperator::EDataType::NotInlined in any of the 2 following cases occur:
+	 * - If the tensor is an EGPUDataType::Output.
+	 * - If the tensor is used by another layer.
+	 * It will return FNeuralOperator::EDataType::Inlined otherwise.
+	 */
+	static bool CanTensorBeInlined(const FString& InTensorName, const TMap<FString, int32>& InInputNameDummyIndexMap,
+		const TMap<FString, int32>& InOutputNameDummyIndexMap, const FNodeProto& InNodeProto, const FGraphProto& InGraphProto, const bool bCanOutputPropertiesChangeWithPostForward = false);
+	/**
+	 * It will return which input tensors can be inlined.
+	 */
+	static TSet<uint32> GetPotentiallyInlinedTensors(const TArray<FString>& InTensorNames, const TMap<FString, int32>& InInputNameDummyIndexMap,
+		const TMap<FString, int32>& InOutputNameDummyIndexMap, const FNodeProto& InNodeProto, const FGraphProto& InGraphProto);
+
+	/**
+	 * Given the input arguments, it creates the proper TShared<FNeuralOperator> child instance.
+	 */
+	static TSharedPtr<FNeuralOperator> CreateOperator(const FNodeProto& InNodeProto, const TArray<FString>& InInputTensorNamesForOperator,
+		const TMap<FString, int32>& InInputNameDummyIndexMap, const TMap<FString, int32>& InOutputNameDummyIndexMap, const FGraphProto& InGraphProto);
+};
diff --git a/Engine/Plugins/Experimental/NeuralNetworkInference/Source/NeuralNetworkInference/Private/UEOnly/NeuralTensorManager.cpp b/Engine/Plugins/Experimental/NeuralNetworkInference/Source/NeuralNetworkInference/Private/UEOnly/NeuralTensorManager.cpp
new file mode 100644
--- /dev/null
+++ b/Engine/Plugins/Experimental/NeuralNetworkInference/Source/NeuralNetworkInference/Private/UEOnly/NeuralTensorManager.cpp
@@ -0,0 +1,373 @@
+// Copyright Epic Games, Inc. All Rights Reserved.
+
+#include "NeuralTensorManager.h"
+#include "NeuralNetworkInferenceUtils.h"
+
+
+
+/* FPrivateNeuralTensorManager
+ *****************************************************************************/
+
+class FPrivateNeuralTensorManager
+{
+public:
+	static TMap<FString, int32> TensorArrayToNameIndexMap(const TArray<FNeuralTensor>& InTensors, const TArray<FNeuralTensor*>& InTensorArray, const FString& InPrefix);
+	static void TensorNameIndexMapToArray(TArray<int32>& OutIndexes, const TMap<FString, int32>& InNameIndexMap);
+};
+
+TMap<FString, int32> FPrivateNeuralTensorManager::TensorArrayToNameIndexMap(const TArray<FNeuralTensor>& InTensors, const TArray<FNeuralTensor*>& InTensorArray, const FString& InPrefix)
+{
+	if (InTensors.Num() >= InTensorArray.Num())
+	{
+		TMap<FString, int32> NameIndexMap;
+		int32 Counter = -1;
+		for (const FNeuralTensor* const Tensor : InTensorArray)
+		{
+			for (int32 TensorIndex = 0; TensorIndex < InTensors.Num(); ++TensorIndex)
+			{
+				if (&InTensors[TensorIndex] == Tensor)
+				{
+					NameIndexMap.Add(InPrefix + FString::FromInt(++Counter), TensorIndex);
+					break;
+				}
+			}
+		}
+		if (NameIndexMap.Num() != InTensorArray.Num())
+		{
+			UE_LOG(LogNeuralNetworkInference, Warning, TEXT("FNeuralTensorManager::TensorArrayToNameIndexMap(): NameIndexMap.Num() == InTensorArray.Num() failed (%d != %d)."),
+				NameIndexMap.Num(), InTensorArray.Num());
+			return TMap<FString, int32>();
+		}
+		return NameIndexMap;
+	}
+	else
+	{
+		UE_LOG(LogNeuralNetworkInference, Warning, TEXT("FNeuralTensorManager::TensorArrayToNameIndexMap(): InTensors.Num() >= InTensorArray.Num() failed (%d vs. %d)."),
+			InTensors.Num(), InTensorArray.Num());
+		return TMap<FString, int32>();
+	}
+}
+
+void FPrivateNeuralTensorManager::TensorNameIndexMapToArray(TArray<int32>& OutIndexes, const TMap<FString, int32>& InNameIndexMap)
+{
+	OutIndexes.Empty();
+	for (const auto& NameIndexPair : InNameIndexMap)
+	{
+		// Add if tensor index only once
+		if (OutIndexes.Find(NameIndexPair.Value) == INDEX_NONE)
+		{
+			OutIndexes.Push(NameIndexPair.Value);
+		}
+	}
+	Algo::Sort(OutIndexes);
+}
+
+
+
+/* FNeuralTensorManager structors
+ *****************************************************************************/
+
+FNeuralTensorManager::FNeuralTensorManager()
+	: bIsLoaded(false)
+{
+}
+
+FNeuralTensorManager::FNeuralTensorManager(TArray<FNeuralTensor>& InTensors, TMap<FString, int32>& InNameIndexMap, TMap<FString, int32>& InInputNameIndexMap,
+	TMap<FString, int32>& InOutputNameIndexMap)
+	: FNeuralTensorManager()
+{
+	Load(InTensors, InNameIndexMap, InInputNameIndexMap, InOutputNameIndexMap);
+}
+
+FNeuralTensorManager::FNeuralTensorManager(TArray<FNeuralTensor>& InTensors, const TArray<FNeuralTensor*>& InInputTensors,
+	const TArray<FNeuralTensor*>& InOutputTensors)
+	: FNeuralTensorManager()
+{
+	Load(InTensors, InInputTensors, InOutputTensors);
+}
+
+FNeuralTensorManager::~FNeuralTensorManager()
+{
+}
+
+
+
+/* FNeuralTensorManager public functions
+ *****************************************************************************/
+
+bool FNeuralTensorManager::Load(TArray<FNeuralTensor>& InTensors, TMap<FString, int32>& InNameIndexMap,
+	TMap<FString, int32>& InInputNameIndexMap, TMap<FString, int32>& InOutputNameIndexMap)
+{
+	bIsLoaded = false;
+	// Sanity checks
+	if (InTensors.Num() < InInputNameIndexMap.Num() || InTensors.Num() < InOutputNameIndexMap.Num())
+	{
+		UE_LOG(LogNeuralNetworkInference, Warning,
+			TEXT("FNeuralTensorManager::Load(): InTensors.Num() >= InputNameIndexMap.Num() (%d vs. %d) || InTensors.Num() >= OutputNameIndexMap.Num() (%d vs. %d) failed."),
+			InTensors.Num(), InInputNameIndexMap.Num(), InTensors.Num(), InOutputNameIndexMap.Num());
+		return bIsLoaded;
+	}
+	else if (InInputNameIndexMap.Num() < 1 || InOutputNameIndexMap.Num() < 1)
+	{
+		UE_LOG(LogNeuralNetworkInference, Warning, TEXT("There should be at least 1 input (currently %d) and 1 output (currently %d) tensors."), InInputNameIndexMap.Num(), InOutputNameIndexMap.Num());
+		return bIsLoaded;
+	}
+	// Fill NameIndexMap, Tensors, InputNameIndexMap, and OutputNameIndexMap
+	Swap(Tensors, InTensors);
+	Swap(NameIndexMap, InNameIndexMap);
+	Swap(InputNameIndexMap, InInputNameIndexMap);
+	Swap(OutputNameIndexMap, InOutputNameIndexMap);
+	// InputIndexes & OutputIndexes
+	FPrivateNeuralTensorManager::TensorNameIndexMapToArray(InputIndexes, InputNameIndexMap);
+	FPrivateNeuralTensorManager::TensorNameIndexMapToArray(OutputIndexes, OutputNameIndexMap);
+	// NonInputIndexes
+	NonInputIndexes.Empty();
+	for (int32 TensorIndex = 0; TensorIndex < Tensors.Num(); ++TensorIndex )
+	{
+		// Add if tensor index is not an input
+		if (InputIndexes.Find(TensorIndex) == INDEX_NONE)
+		{
+			NonInputIndexes.Push(TensorIndex);
+		}
+	}
+	if (InputIndexes.Num() + NonInputIndexes.Num() != Tensors.Num())
+	{
+		UE_LOG(LogNeuralNetworkInference, Warning,
+			TEXT("FNeuralTensorManager::Load(): InputIndexes.Num() + NonInputIndexes.Num() == Tensors.Num() failed (%d + %d != %d), NameIndexMap.Num() = %d)."),
+			InputIndexes.Num(), NonInputIndexes.Num(), Tensors.Num(), NameIndexMap.Num());
+		return bIsLoaded;
+	}
+	// FNeuralTensorManager was configured
+	bIsLoaded = true;
+	return bIsLoaded;
+}
+
+bool FNeuralTensorManager::Load(TArray<FNeuralTensor>& InTensors, const TArray<FNeuralTensor*>& InInputTensors,
+	const TArray<FNeuralTensor*>& InOutputTensors)
+{
+	TMap<FString, int32> InputNameIndexMapLocal = FPrivateNeuralTensorManager::TensorArrayToNameIndexMap(InTensors, InInputTensors, TEXT("i"));
+	TMap<FString, int32> OutputNameIndexMapLocal = FPrivateNeuralTensorManager::TensorArrayToNameIndexMap(InTensors, InOutputTensors, TEXT("o"));
+	TMap<FString, int32> NameIndexMapLocal = InputNameIndexMapLocal;
+	NameIndexMapLocal.Append(OutputNameIndexMapLocal);
+	// Missing FNeuralTensors (weight & intermediate ones)
+	if (NameIndexMapLocal.Num() < InTensors.Num())
+	{
+		TArray<int32> IndexesLocal;
+		FPrivateNeuralTensorManager::TensorNameIndexMapToArray(IndexesLocal, NameIndexMapLocal);
+		int32 Counter = -1;
+		for (int32 TensorIndex = 0; TensorIndex < InTensors.Num(); ++TensorIndex)
+		{
+			// If not an output nor input --> Add with new name
+			if (IndexesLocal.Find(TensorIndex) == INDEX_NONE)
+			{
+				NameIndexMap.Add(TEXT("t") + FString::FromInt(++Counter), TensorIndex);
+			}
+		}
+	}
+	// Call the original Load() with the new TMaps
+	return Load(InTensors, NameIndexMapLocal, InputNameIndexMapLocal, OutputNameIndexMapLocal);
+}
+
+bool FNeuralTensorManager::IsLoaded() const
+{
+	return bIsLoaded;
+}
+
+const TArray<FNeuralTensor>& FNeuralTensorManager::GetTensors() const
+{
+	SanityCheckIsLoaded();
+	return Tensors;
+}
+
+TArray<FNeuralTensor>& FNeuralTensorManager::GetTensorsMutable()
+{
+	SanityCheckIsLoaded();
+	return Tensors;
+}
+
+const TMap<FString, int32>& FNeuralTensorManager::GetNameIndexMap() const
+{
+	SanityCheckIsLoaded();
+	return NameIndexMap;
+}
+
+const TArray<int32>& FNeuralTensorManager::GetNonInputIndexes() const
+{
+	SanityCheckIsLoaded();
+	return NonInputIndexes;
+}
+
+const FNeuralTensor& FNeuralTensorManager::GetInputTensor() const
+{
+	SanityCheckIsLoaded();
+	if (InputIndexes.Num() != 1)
+	{
+		UE_LOG(LogNeuralNetworkInference, Warning,
+			TEXT("FNeuralTensorManager::GetInputTensor(): There should be 1 input, not %d. Call GetInputIndexes() instead."), InputIndexes.Num());
+	}
+	checkf(Tensors.Num() > 0 && InputIndexes.Num() > 0, TEXT("Tensors and/or InputIndexes was empty."));
+	return Tensors[InputIndexes[0]];
+}
+
+FNeuralTensor& FNeuralTensorManager::GetInputTensorMutable()
+{
+	SanityCheckIsLoaded();
+	if (InputIndexes.Num() != 1)
+	{
+		UE_LOG(LogNeuralNetworkInference, Warning,
+			TEXT("FNeuralTensorManager::GetInputTensorMutable(): There should be 1 input, not %d. Call GetInputIndexes() instead."), InputIndexes.Num());
+	}
+	checkf(Tensors.Num() > 0 && InputIndexes.Num() > 0, TEXT("Tensors and/or InputIndexes was empty."));
+	return Tensors[InputIndexes[0]];
+}
+
+const TArray<int32>& FNeuralTensorManager::GetInputIndexes() const
+{
+	SanityCheckIsLoaded();
+	return InputIndexes;
+}
+
+const TMap<FString, int32>& FNeuralTensorManager::GetInputNameIndexMap() const
+{
+	SanityCheckIsLoaded();
+	return InputNameIndexMap;
+}
+
+void FNeuralTensorManager::SetInputFromTensorCopy(const FNeuralTensor& InTensor)
+{
+	return GetInputTensorMutable().SetFromUnderlyingUInt8ArrayCopy(InTensor.GetUnderlyingUInt8ArrayRef());
+}
+
+void FNeuralTensorManager::SetInputFromTensorMapCopy(const TMap<FString, FNeuralTensor>& InTensorMap)
+{
+	TArray<FNeuralTensor>& TensorsLocalRef = GetTensorsMutable();
+	for (const auto& IndexMapTuple : GetInputNameIndexMap())
+	{
+		TensorsLocalRef[IndexMapTuple.Value].SetFromUnderlyingUInt8ArrayCopy(InTensorMap.FindChecked(IndexMapTuple.Key).GetUnderlyingUInt8ArrayRef());
+	}
+}
+
+TMap<FString, void*> FNeuralTensorManager::CreateInputDataPointersMutable()
+{
+	TMap<FString, void*> InputDataPointers;
+	if (SanityCheckIsLoaded())
+	{
+		for (auto& NameIndexPair : GetInputNameIndexMap())
+		{
+			InputDataPointers.Add(NameIndexPair.Key, Tensors[NameIndexPair.Value].GetData());
+		}
+	}
+	return InputDataPointers;
+}
+
+FRDGBufferUAVRef FNeuralTensorManager::GetInputBufferUAVRef()
+{
+	return GetInputTensorMutable().GetBufferUAVRef();
+}
+
+TMap<FString, FRDGBufferUAVRef> FNeuralTensorManager::CreateInputBufferUAVRefs()
+{
+	TMap<FString, FRDGBufferUAVRef> InputBufferUAVMap;
+	if (SanityCheckIsLoaded())
+	{
+		for (auto& NameIndexPair : GetInputNameIndexMap())
+		{
+			InputBufferUAVMap.Add(NameIndexPair.Key, Tensors[NameIndexPair.Value].GetBufferUAVRef());
+		}
+	}
+	return InputBufferUAVMap;
+}
+
+const FNeuralTensor& FNeuralTensorManager::GetOutputTensor() const
+{
+	SanityCheckIsLoaded();
+	if (OutputIndexes.Num() != 1)
+	{
+		UE_LOG(LogNeuralNetworkInference, Warning,
+			TEXT("FNeuralTensorManager::GetOutputTensor(): There should be 1 input, not %d. Call GetOutputIndexes() instead."), OutputIndexes.Num());
+	}
+	checkf(Tensors.Num() > 0 && OutputIndexes.Num() > 0, TEXT("Tensors and/or OutputIndexes was empty."));
+	return Tensors[OutputIndexes[0]];
+}
+
+FNeuralTensor& FNeuralTensorManager::GetOutputTensorMutable()
+{
+	SanityCheckIsLoaded();
+	if (OutputIndexes.Num() != 1)
+	{
+		UE_LOG(LogNeuralNetworkInference, Warning,
+			TEXT("FNeuralTensorManager::GetOutputTensorMutable(): There should be 1 input, not %d. Call GetOutputIndexes() instead."), OutputIndexes.Num());
+	}
+	checkf(Tensors.Num() > 0 && OutputIndexes.Num() > 0, TEXT("Tensors and/or OutputIndexes was empty."));
+	return Tensors[OutputIndexes[0]];
+}
+
+const TArray<int32>& FNeuralTensorManager::GetOutputIndexes() const
+{
+	SanityCheckIsLoaded();
+	return OutputIndexes;
+}
+
+const TMap<FString, int32>& FNeuralTensorManager::GetOutputNameIndexMap() const
+{
+	SanityCheckIsLoaded();
+	return OutputNameIndexMap;
+}
+
+const FRDGBufferSRVRef FNeuralTensorManager::GetOutputBufferSRVRef() const
+{
+	return GetOutputTensor().GetBufferSRVRef();
+}
+
+TMap<FString, const FRDGBufferSRVRef> FNeuralTensorManager::CreateOutputBufferSRVRefs() const
+{
+	TMap<FString, const FRDGBufferSRVRef> OutputBufferSRVMap;
+	if (SanityCheckIsLoaded())
+	{
+		for (const auto& NameIndexPair : GetOutputNameIndexMap())
+		{
+			OutputBufferSRVMap.Add(NameIndexPair.Key, Tensors[NameIndexPair.Value].GetBufferSRVRef());
+		}
+	}
+	return OutputBufferSRVMap;
+}
+
+TMap<FString, FNeuralTensor> FNeuralTensorManager::CreateInputTensorMap() const
+{
+	TMap<FString, FNeuralTensor> InputArrayMap;
+	if (SanityCheckIsLoaded())
+	{
+		for (const auto& NameIndexPair : GetInputNameIndexMap())
+		{
+			InputArrayMap.Add(NameIndexPair.Key, Tensors[NameIndexPair.Value]);
+		}
+	}
+	return InputArrayMap;
+}
+
+TMap<FString, FNeuralTensor> FNeuralTensorManager::CreateOutputTensorMap() const
+{
+	TMap<FString, FNeuralTensor> OutputArrayMap;
+	if (SanityCheckIsLoaded())
+	{
+		for (const auto& NameIndexPair : GetOutputNameIndexMap())
+		{
+			OutputArrayMap.Add(NameIndexPair.Key, Tensors[NameIndexPair.Value]);
+		}
+	}
+	return OutputArrayMap;
+}
+
+
+
+/* FNeuralTensorManager private functions
+ *****************************************************************************/
+
+bool FNeuralTensorManager::SanityCheckIsLoaded() const
+{
+	if (!bIsLoaded)
+	{
+		UE_LOG(LogNeuralNetworkInference, Warning, TEXT("FNeuralTensorManager::SanityCheckIsLoaded(): Not loaded."));
+	}
+	return bIsLoaded;
+}
diff --git a/Engine/Plugins/Experimental/NeuralNetworkInference/Source/NeuralNetworkInference/Private/UEOnly/NeuralTensorManager.h b/Engine/Plugins/Experimental/NeuralNetworkInference/Source/NeuralNetworkInference/Private/UEOnly/NeuralTensorManager.h
new file mode 100644
--- /dev/null
+++ b/Engine/Plugins/Experimental/NeuralNetworkInference/Source/NeuralNetworkInference/Private/UEOnly/NeuralTensorManager.h
@@ -0,0 +1,190 @@
+// Copyright Epic Games, Inc. All Rights Reserved.
+
+#pragma once
+
+#include "CoreMinimal.h"
+#include "NeuralTensor.h"
+
+/**
+ * This is an auxiliary class. See UNeuralNetworkLegacy for a high-level wrapper of the whole NeuralNetworkInference plugin. The UNeuralNetworkLegacy header
+ * documentation also includes some code examples.
+ *
+ * FNeuralTensorManager is an utility class that helps store all the FNeuralTensors related to a particular UNeuralNetworkLegacy model.
+ * Most functions are inlined (or templated) to avoid the performance overhead of a wrapper class.
+ */
+class FNeuralTensorManager
+{
+public:
+	FNeuralTensorManager();
+
+	/**
+	 * Calls FNeuralTensorManager() + Load(InNameIndexMap, InTensors, InInputNameIndexMap, InOutputNameIndexMap).
+	 */
+	FNeuralTensorManager(TArray<FNeuralTensor>& InTensors, TMap<FString, int32>& InNameIndexMap, TMap<FString, int32>& InInputNameIndexMap,
+		TMap<FString, int32>& InOutputNameIndexMap);
+
+	/**
+	 * Calls FNeuralTensorManager() + Load(InNameIndexMap, InTensors, InInputTensors, InOutputTensors).
+	 */
+	FNeuralTensorManager(TArray<FNeuralTensor>& InTensors, const TArray<FNeuralTensor*>& InInputTensors, const TArray<FNeuralTensor*>& InOutputTensors);
+
+	~FNeuralTensorManager();
+
+	/**
+	 * It moves the TMap objects and populates all other class members.
+	 */
+	bool Load(TArray<FNeuralTensor>& InTensors, TMap<FString, int32>& InNameIndexMap, TMap<FString, int32>& InInputNameIndexMap,
+		TMap<FString, int32>& InOutputNameIndexMap);
+
+	/**
+	 * Alternative to the original Load(), recommended for single-input single-output models as well as any kind of network whose
+	 * input and output tensor names are irrelevant). It will internally create InNameIndexMap/InInputNameIndexMap/InOutputNameIndexMap
+	 * and call the original Load().
+	 * @param InInputTensors Pointers for the tensors in InTensors that are inputs.
+	 * @param InOutputTensors Pointers for the tensors in InTensors that are outputs.
+	 */
+	bool Load(TArray<FNeuralTensor>& InTensors, const TArray<FNeuralTensor*>& InInputTensors, const TArray<FNeuralTensor*>& InOutputTensors);
+
+	/**
+	 * It returns whether it loaded the UNeuralNetworkLegacy successfully.
+	 */
+	bool IsLoaded() const;
+
+	/**
+	 * Getter for the Tensors TArray. GetTensors() will return a const reference while GetTensorsMutable() allows modifying it.
+	 */
+	const TArray<FNeuralTensor>& GetTensors() const;
+	TArray<FNeuralTensor>& GetTensorsMutable();
+
+	/**
+	 * Getter for the NameIndexMap TMap.
+	 */
+	const TMap<FString, int32>& GetNameIndexMap() const;
+
+	/**
+	 * It returns the non-input FNeuralTensor index TMap as a read-only object.
+	 */
+	const TArray<int32>& GetNonInputIndexes() const;
+
+	/**
+	 * It returns the input FNeuralTensor(s) as read-only (GetInputTensor/GetInputIndexes/GetInputNameIndexMap) or mutable (GetInputTensorMutable) objects.
+	 * GetInputTensor/GetInputTensorMutable ensure there is only 1 input and returns a single FNeuralTensor&, while GetInputIndexes/GetInputNameIndexMap return them all.
+	 * In order to modify their input data of the FNeuralTensor (map), use GetInputDataPointerMutable() or CreateInputDataPointersMutable().
+	 *
+	 * Use the read-only versions as much as possible to avoid potential undefined behavior in the next Run() of the FNeuralTensorManager due to non-controlled functions
+	 * (e.g., resizes, memory re-allocation).
+	 *
+	 * @return GetInputTensor() returns a read-only const FNeuralTensor& corresponding to the single input FNeuralTensor, GetInputIndexes()/GetInputNameIndexMap() return a
+	 * TArray<int32>/TMap<FString, int32> for each one of the input FNeuralTensors. They will all go out of scope if FNeuralTensorManager does.
+	 */
+	const FNeuralTensor& GetInputTensor() const;
+	FNeuralTensor& GetInputTensorMutable();
+	const TArray<int32>& GetInputIndexes() const;
+	const TMap<FString, int32>& GetInputNameIndexMap() const;
+
+	/**
+	 * There are 6 alternative functions to fill the input FNeuralTensor(s) data:
+	 * If exactly 1 input FNeuralTensor: SetInputFromArrayCopy(), SetInputFromTensorCopy(), GetInputDataPointerMutable().
+	 *     - They ensure there is exactly 1 input FNeuralTensor or will log a warning if more than 1 input tensor exists.
+	 * If more than 1 input FNeuralTensors: SetInputFromTensorMapCopy(), CreateInputDataPointersMutable().
+	 *
+	 * - SetInputFromArrayCopy()/SetInputFromTensorCopy()/SetInputFromTensorMapCopy() deeply copy the input FNeuralTensor(s) (slower but safer and less error
+	 *   prone). See FNeuralTensor::SetFromArrayCopy() for more details.
+	 * - GetInputDataPointerMutable() returns a pointer to the input FNeuralTensor raw data, so it can be filled before calling Run().
+	 * - CreateInputDataPointersMutable() returns a TMap of pointers to the input FNeuralTensors raw data, so it can be filled before calling Run().
+	 *
+	 * @return GetInputDataPointerMutable() returns a T* corresponding to the mutable data in the single input FNeuralTensor, CreateInputDataPointersMutable() returns a
+	 * TMap<FString, void*> for each one of the input FNeuralTensors.
+	 *
+	 * For read-only access to the input FNeuralTensor(s), see GetInputTensor() or GetInputNameIndexMap() (e.g., to extract properties
+	 * from the input FNeuralTensor(s) such as volume, dimensions).
+	 */
+	template<typename T>
+	void SetInputFromArrayCopy(const TArray<T>& InArray);
+	void SetInputFromTensorCopy(const FNeuralTensor& InTensor);
+	void SetInputFromTensorMapCopy(const TMap<FString, FNeuralTensor>& InTensorMap);
+	template<typename T>
+	T* GetInputDataPointerMutable();
+	TMap<FString, void*> CreateInputDataPointersMutable();
+	FRDGBufferUAVRef GetInputBufferUAVRef();
+	TMap<FString, FRDGBufferUAVRef> CreateInputBufferUAVRefs();
+
+	/**
+	 * It returns the input FNeuralTensor(s) as read-only (GetOutputTensor/GetOutputIndexes/GetOutputNameIndexMap) or mutable (GetOutputTensorMutable) objects.
+	 * GetOutputTensor/GetOutputTensorMutable ensure there is only 1 input and returns a single FNeuralTensor&, while GetOutputIndexes/GetOutputNameIndexMap return them all.
+	 * In order to modify their input data of the FNeuralTensor (map), use GetOutputDataPointerMutable() or CreateOutputDataPointersMutable().
+	 *
+	 * Use the read-only versions as much as possible to avoid potential undefined behavior in the next Run() of the FNeuralTensorManager due to non-controlled functions
+	 * (e.g., resizes, memory re-allocation).
+	 *
+	 * @return GetOutputTensor() returns a read-only const FNeuralTensor& corresponding to the single input FNeuralTensor, GetOutputIndexes()/GetOutputNameIndexMap() return a
+	 * TArray<int32>/TMap<FString, int32> for each one of the input FNeuralTensors. They will all go out of scope if FNeuralTensorManager does.
+	 */
+	const FNeuralTensor& GetOutputTensor() const;
+	FNeuralTensor& GetOutputTensorMutable();
+	const TArray<int32>& GetOutputIndexes() const;
+	const TMap<FString, int32>& GetOutputNameIndexMap() const;
+	const FRDGBufferSRVRef GetOutputBufferSRVRef() const;
+	TMap<FString, const FRDGBufferSRVRef> CreateOutputBufferSRVRefs() const;
+
+	/**
+	 * These functions are slower than the rest, because they deep copy each one of the TArray's in the final TMap.
+	 */
+	TMap<FString, FNeuralTensor> CreateInputTensorMap() const;
+	TMap<FString, FNeuralTensor> CreateOutputTensorMap() const;
+
+protected:
+	UPROPERTY(VisibleAnywhere, Category = "Neural Network Inference")
+	bool bIsLoaded;
+	/**
+	 * It contains the raw weight, intermediate, input, and output FNeuralTensor's. Inlined FNeuralTensor's will only appear once on Tensors.
+	 * E.g., in the example described for NameIndexMap, the FNeuralTensor will be present with the Key name "Input" (or otherwise "Output"). E.g.,
+	 * FNeuralTensor& Tensor = Tensors[NameIndexMap.FindChecked(TEXT("Input"))];
+	 */
+	UPROPERTY(VisibleAnywhere, Category = "Neural Network Inference")
+	TArray<FNeuralTensor> Tensors;
+	/**
+	 * For each original FNeuralTensor name, it saves its final name after merging FNeuralTensor's for memory savings.
+	 * E.g., For an inlined ReLU operator with original input tensor name "Input" and output tensor name "Output", NameIndexMap will
+	 * contain the following 2 fields: {"Input", 0} and {"Output", 0}.
+	 * E.g., For a non-inlined ReLU: {"Input", 0} and {"Output", 1}.
+	 */
+	UPROPERTY(VisibleAnywhere, Category = "Neural Network Inference")
+	TMap<FString, int32> NameIndexMap;
+	/**
+	 * Other variables given in Load().
+	 */
+	UPROPERTY(VisibleAnywhere, Category = "Neural Network Inference")
+	TMap<FString, int32> InputNameIndexMap;
+	UPROPERTY(VisibleAnywhere, Category = "Neural Network Inference")
+	TMap<FString, int32> OutputNameIndexMap;
+	/**
+	 * Other variables auto-generated in Load().
+	 */
+	UPROPERTY(VisibleAnywhere, Category = "Neural Network Inference")
+	TArray<int32> InputIndexes;
+	UPROPERTY(VisibleAnywhere, Category = "Neural Network Inference")
+	TArray<int32> NonInputIndexes;
+	UPROPERTY(VisibleAnywhere, Category = "Neural Network Inference")
+	TArray<int32> OutputIndexes;
+
+	bool SanityCheckIsLoaded() const;
+};
+
+
+
+/* FNeuralTensorManager inline functions
+ *****************************************************************************/
+
+template<typename T>
+void FNeuralTensorManager::SetInputFromArrayCopy(const TArray<T>& InArray)
+{
+	return GetInputTensorMutable().SetFromArrayCopy(InArray);
+}
+
+template<typename T>
+T* FNeuralTensorManager::GetInputDataPointerMutable()
+{
+	// GetInputTensorMutable() already runs the sanity checks
+	return GetInputTensorMutable().GetDataCasted<T>();
+}
diff --git a/Engine/Plugins/Experimental/NeuralNetworkInference/Source/NeuralNetworkInference/Public/NeuralNetwork.h b/Engine/Plugins/Experimental/NeuralNetworkInference/Source/NeuralNetworkInference/Public/NeuralNetwork.h
--- a/Engine/Plugins/Experimental/NeuralNetworkInference/Source/NeuralNetworkInference/Public/NeuralNetwork.h
+++ b/Engine/Plugins/Experimental/NeuralNetworkInference/Source/NeuralNetworkInference/Public/NeuralNetwork.h
@@ -268,7 +268,7 @@ public:
 	 * Internal function not needed by the user.
 	 * Used to create custom networks without an ONNX file for QA testing in FOperatorTester::TestOperator().
 	 */
-	bool Load(TSharedPtr<struct FNeuralTensorManager>& InTensorManager, const TArray<TSharedPtr<class FNeuralOperator>>& InOperators);
+	bool Load(TArray<FNeuralTensor>& InTensors, const TArray<FNeuralTensor*>& InInputTensors, const TArray<FNeuralTensor*>& InOutputTensors, const TArray<TSharedPtr<class FNeuralOperator>>& InOperators);
 
 #if WITH_EDITOR
 	/**
diff --git a/Engine/Plugins/Experimental/NeuralNetworkInference/Source/NeuralNetworkInferenceQA/Private/OperatorTester.cpp b/Engine/Plugins/Experimental/NeuralNetworkInference/Source/NeuralNetworkInferenceQA/Private/OperatorTester.cpp
--- a/Engine/Plugins/Experimental/NeuralNetworkInference/Source/NeuralNetworkInferenceQA/Private/OperatorTester.cpp
+++ b/Engine/Plugins/Experimental/NeuralNetworkInference/Source/NeuralNetworkInferenceQA/Private/OperatorTester.cpp
@@ -3,7 +3,6 @@
 #include "OperatorTester.h"
 #include "NeuralNetworkInferenceQAUtils.h"
 #include "NeuralOperators.h"
-#include "NeuralTensorManager.h"
 #include "Misc/Paths.h"
 
 
@@ -82,13 +81,7 @@ void FOperatorTester::TestOperator(UNeuralNetworkInferenceQAAsset* InOutNetworkI
 			ensureMsgf(false, TEXT("Network is nullptr."));
 			return;
 		}
-		TSharedPtr<FNeuralTensorManager> TensorManager = MakeShared<FNeuralTensorManager>(InOutTensors[Index], InputTensors, OutputTensors);
-		if (!TensorManager->IsLoaded())
-		{
-			ensureMsgf(false, TEXT("TensorManager could not be loaded."));
-			return;
-		}
-		if (!Network->Load(TensorManager, { Operator }))
+		if (!Network->Load(InOutTensors[Index], InputTensors, OutputTensors, { Operator }))
 		{
 			ensureMsgf(false, TEXT("Network could not be loaded."));
 			return;
