Fixed Automation tests for ORT - Static Analysis.
#rnx

diff --git a/Engine/Plugins/Experimental/NeuralNetworkInference/Source/Internal/ThirdPartyHelperAndDLLLoader/Internal/ThirdPartyWarningDisabler.h b/Engine/Plugins/Experimental/NeuralNetworkInference/Source/Internal/ThirdPartyHelperAndDLLLoader/Internal/ThirdPartyWarningDisabler.h
--- a/Engine/Plugins/Experimental/NeuralNetworkInference/Source/Internal/ThirdPartyHelperAndDLLLoader/Internal/ThirdPartyWarningDisabler.h
+++ b/Engine/Plugins/Experimental/NeuralNetworkInference/Source/Internal/ThirdPartyHelperAndDLLLoader/Internal/ThirdPartyWarningDisabler.h
@@ -28,10 +28,12 @@ Many third party headers require some care when importing. NNI third party inclu
 	__pragma(warning(disable: 6001)) /* For ONNX Runtime */ \
 	__pragma(warning(disable: 6246)) /* For ONNX Runtime */ \
 	__pragma(warning(disable: 6258)) /* For ONNX Runtime */ \
+	__pragma(warning(disable: 6294)) /* For ONNX Runtime */ \
 	__pragma(warning(disable: 6313)) /* For ONNX Runtime */ \
 	__pragma(warning(disable: 6387)) /* For ONNX Runtime */ \
 	__pragma(warning(disable: 6388)) /* For ONNX Runtime */ \
 	__pragma(warning(disable: 6504)) /* For ONNX Runtime */ \
+	__pragma(warning(disable: 28020)) /* For ONNX Runtime */ \
 	__pragma(warning(disable: 28196)) /* For ONNX Runtime */ \
 	__pragma(warning(disable: 28204)) /* For ONNX Runtime */ \
 	__pragma(warning(disable: 28205)) /* For ONNX Runtime */ \
