Reduced number disabled warnings for ORT.

#rnx

diff --git a/Engine/Plugins/Experimental/NeuralNetworkInference/Source/Internal/ThirdPartyHelperAndDLLLoader/Internal/ThirdPartyWarningDisabler.h b/Engine/Plugins/Experimental/NeuralNetworkInference/Source/Internal/ThirdPartyHelperAndDLLLoader/Internal/ThirdPartyWarningDisabler.h
--- a/Engine/Plugins/Experimental/NeuralNetworkInference/Source/Internal/ThirdPartyHelperAndDLLLoader/Internal/ThirdPartyWarningDisabler.h
+++ b/Engine/Plugins/Experimental/NeuralNetworkInference/Source/Internal/ThirdPartyHelperAndDLLLoader/Internal/ThirdPartyWarningDisabler.h
@@ -16,28 +16,14 @@ Many third party headers require some care when importing. NNI third party inclu
 
 #ifdef PLATFORM_NNI_MICROSOFT
 #define NNI_THIRD_PARTY_INCLUDES_START THIRD_PARTY_INCLUDES_START \
-	__pragma(warning(disable: 4127)) /* For ONNX Runtime */ \
-	__pragma(warning(disable: 4191)) /* For ONNX Runtime */ \
-	__pragma(warning(disable: 4471)) /* For ONNX Runtime */ \
-	__pragma(warning(disable: 4495)) /* For ONNX Runtime */ \
-	__pragma(warning(disable: 4497)) /* For ONNX Runtime */ \
-	__pragma(warning(disable: 4530)) /* For ONNX Runtime */ \
-	__pragma(warning(disable: 4800)) /* For ONNX Runtime */ \
-	__pragma(warning(disable: 4815)) /* For ONNX Runtime */ \
-	__pragma(warning(disable: 4834)) /* For ONNX Runtime */ \
-	__pragma(warning(disable: 4946)) /* For ONNX Runtime */ \
-	__pragma(warning(disable: 6001)) /* For ONNX Runtime */ \
-	__pragma(warning(disable: 6246)) /* For ONNX Runtime */ \
-	__pragma(warning(disable: 6258)) /* For ONNX Runtime */ \
-	__pragma(warning(disable: 6294)) /* For ONNX Runtime */ \
-	__pragma(warning(disable: 6313)) /* For ONNX Runtime */ \
-	__pragma(warning(disable: 6387)) /* For ONNX Runtime */ \
-	__pragma(warning(disable: 6388)) /* For ONNX Runtime */ \
-	__pragma(warning(disable: 6504)) /* For ONNX Runtime */ \
-	__pragma(warning(disable: 28020)) /* For ONNX Runtime */ \
-	__pragma(warning(disable: 28196)) /* For ONNX Runtime */ \
-	__pragma(warning(disable: 28204)) /* For ONNX Runtime */ \
-	__pragma(warning(disable: 28205)) /* For ONNX Runtime */ \
+	__pragma(warning(disable: 4191)) /* C4191: 'reinterpret_cast': unsafe conversion from 'FARPROC' to 'void (__cdecl *)(PEXCEPTION_RECORD,PCONTEXT,DWORD) */ \
+	__pragma(warning(disable: 4497)) /* C4497: nonstandard extension 'sealed' used: replace with 'final' */ \
+	__pragma(warning(disable: 6001)) /* C6001: Using uninitialized memory '*pNode'. */ \
+	__pragma(warning(disable: 6294)) /* C6294: Ill-defined for-loop:  initial condition does not satisfy test.  Loop body not executed. */ \
+	__pragma(warning(disable: 6313)) /* C6313: Incorrect operator:  zero-valued flag cannot be tested with bitwise-and.  Use an equality test to check for zero-valued flags. */ \
+	__pragma(warning(disable: 6387)) /* C6387: 'X' could be '0':  this does not adhere to the specification for the function 'memset'. */ \
+	__pragma(warning(disable: 6388)) /* C6388: '*X' might not be '0':  this does not adhere to the specification for the function 'Y'. */ \
+	__pragma(warning(disable: 28020)) /* C28020: The expression '0<=_Param_(1)&&_Param_(1)<=400-1' is not true at this call. */ \
 	UE_PUSH_MACRO("check") \
 	UE_PUSH_MACRO("TEXT")
 #else
