NeuralNetworkInference plugin:
Neural Network Inference Time stats

diff --git a/Engine/Plugins/Experimental/NeuralNetworkInference/NeuralNetworkInference.uplugin b/Engine/Plugins/Experimental/NeuralNetworkInference/NeuralNetworkInference.uplugin
--- a/Engine/Plugins/Experimental/NeuralNetworkInference/NeuralNetworkInference.uplugin
+++ b/Engine/Plugins/Experimental/NeuralNetworkInference/NeuralNetworkInference.uplugin
@@ -35,12 +35,19 @@
 			"Type": "RuntimeAndProgram",
 			"LoadingPhase": "Default"
 		},
+		{
+			"Name": "NeuralNetworkInferenceProfiling",
+			"Type": "RuntimeAndProgram",
+			"LoadingPhase": "Default"
+		},
 		{
 			"Name": "ONNXRuntime",
 			"Type": "RuntimeAndProgram",
 			"LoadingPhase": "Default",
 			"PlatformAllowList": [
-				"Win64", "Linux", "Mac"
+				"Win64",
+				"Linux",
+				"Mac"
 			]
 		},
 		{
diff --git a/Engine/Plugins/Experimental/NeuralNetworkInference/Source/NeuralNetworkInference/Internal/NeuralNetworkInferenceTimer.cpp b/Engine/Plugins/Experimental/NeuralNetworkInference/Source/NeuralNetworkInference/Internal/NeuralNetworkInferenceTimer.cpp
deleted file mode 100644
diff --git a/Engine/Plugins/Experimental/NeuralNetworkInference/Source/NeuralNetworkInference/NeuralNetworkInference.Build.cs b/Engine/Plugins/Experimental/NeuralNetworkInference/Source/NeuralNetworkInference/NeuralNetworkInference.Build.cs
--- a/Engine/Plugins/Experimental/NeuralNetworkInference/Source/NeuralNetworkInference/NeuralNetworkInference.Build.cs
+++ b/Engine/Plugins/Experimental/NeuralNetworkInference/Source/NeuralNetworkInference/NeuralNetworkInference.Build.cs
@@ -30,6 +30,7 @@ public class NeuralNetworkInference : ModuleRules
 				"CoreUObject",
 				"Engine",
 				"NeuralNetworkInferenceCore",
+				"NeuralNetworkInferenceProfiling",
 				// Internal (Backend-related)
 				"ModelProto"
 			}
diff --git a/Engine/Plugins/Experimental/NeuralNetworkInference/Source/NeuralNetworkInference/Private/NeuralNetwork.cpp b/Engine/Plugins/Experimental/NeuralNetworkInference/Source/NeuralNetworkInference/Private/NeuralNetwork.cpp
--- a/Engine/Plugins/Experimental/NeuralNetworkInference/Source/NeuralNetworkInference/Private/NeuralNetwork.cpp
+++ b/Engine/Plugins/Experimental/NeuralNetworkInference/Source/NeuralNetworkInference/Private/NeuralNetwork.cpp
@@ -5,6 +5,7 @@
 #include "NeuralNetworkImplBackEndUEOnly.h"
 #include "NeuralNetworkInferenceUtils.h"
 #include "EditorFramework/AssetImportData.h"
+#include "NeuralNetworkInferenceTimer.h"
 #include "HAL/FileManager.h"
 #include "Misc/FileHelper.h"
 #include "Misc/Paths.h"
@@ -193,6 +194,13 @@ void UNeuralNetwork::SetBackEnd(const ENeuralBackEnd InBackEnd)
 	}
 }
 
+void UNeuralNetwork::ResetStats()
+{
+	ComputeStatsModule.ResetStats();
+	InputMemoryTransferStatsModule.ResetStats();
+}
+
+
 bool UNeuralNetwork::IsGPUConfigCompatibleForUEAndORTBackEnd()
 {
 	return FImplBackEndUEAndORT::IsGPUConfigCompatible();
@@ -247,7 +255,9 @@ void UNeuralNetwork::SetInputFromArrayCopy(const TArray<float>& InArray, const i
 	{
 		UE_LOG(LogNeuralNetworkInference, Warning, TEXT("UNeuralNetwork::SetInputFromArrayCopy(): Call UNeuralNetwork::Load() to load a model first."));
 	}
-
+	
+	FNeuralNetworkInferenceTimer RunTimer;
+	RunTimer.Tic();
 	// UEAndORT
 	if (BackEndForCurrentPlatform == ENeuralBackEnd::UEAndORT)
 	{
@@ -266,6 +276,9 @@ void UNeuralNetwork::SetInputFromArrayCopy(const TArray<float>& InArray, const i
 	{
 		UE_LOG(LogNeuralNetworkInference, Warning, TEXT("UNeuralNetwork::SetInputFromArrayCopy(): Unknown [BackEnd,BackEndForCurrentPlatform] = [%d,%d]."), (int32)BackEnd, (int32)BackEndForCurrentPlatform);
 	}
+
+	InputMemoryTransferStatsModule.StoreSample(RunTimer.Toc());
+
 }
 
 void* UNeuralNetwork::GetInputDataPointerMutable(const int32 InTensorIndex)
@@ -370,11 +383,14 @@ void UNeuralNetwork::SetInputFromArrayCopy(const TArray<FNeuralTensor>& InInputT
 		UE_LOG(LogNeuralNetworkInference, Warning, TEXT("UNeuralNetwork::SetInputFromArrayCopy(): GetInputTensorNumber() == InInputTensorArray.Num() failed, %d != %d."), GetInputTensorNumber(), InInputTensorArray.Num());
 		return;
 	}
+	FNeuralNetworkInferenceTimer RunTimer;
+	RunTimer.Tic();
 	for (uint32 InputTensorIndex = 0; InputTensorIndex < GetInputTensorNumber(); ++InputTensorIndex)
 	{
 		FNeuralTensor& InputTensor = GetInputTensorMutable(InputTensorIndex);
 		InputTensor.SetFromUnderlyingUInt8ArrayCopy(InInputTensorArray[InputTensorIndex].GetUnderlyingUInt8ArrayRef());
 	}
+	InputMemoryTransferStatsModule.StoreSample(RunTimer.Toc());
 }
 
 TArray<FNeuralTensor> UNeuralNetwork::CreateOutputArrayCopy() const
@@ -482,7 +498,8 @@ void UNeuralNetwork::Run()
 		UE_LOG(LogNeuralNetworkInference, Warning, TEXT("UNeuralNetwork::Run(): Call UNeuralNetwork::Load() to load a model first."));
 		return;
 	}
-
+	FNeuralNetworkInferenceTimer RunTimer;
+	RunTimer.Tic();
 	// UEAndORT
 	if (BackEndForCurrentPlatform == ENeuralBackEnd::UEAndORT)
 	{
@@ -500,8 +517,9 @@ void UNeuralNetwork::Run()
 	{
 		UE_LOG(LogNeuralNetworkInference, Warning, TEXT("UNeuralNetwork::Run(): Unknown [BackEnd,BackEndForCurrentPlatform] = [%d,%d]."), (int32)BackEnd, (int32)BackEndForCurrentPlatform);
 	}
-}
 
+	ComputeStatsModule.StoreSample(RunTimer.Toc());
+}
 
 
 /* UNeuralNetwork private functions
@@ -531,6 +549,9 @@ bool UNeuralNetwork::Load()
 		UE_LOG(LogNeuralNetworkInference, Warning, TEXT("UNeuralNetwork::Load(): Unknown [BackEnd,BackEndForCurrentPlatform] = [%d,%d]."), (int32)BackEnd, (int32)BackEndForCurrentPlatform);
 	}
 
+	// Reset Stats
+	ResetStats();
+
 	return bIsLoaded;
 }
 
@@ -635,3 +656,15 @@ void UNeuralNetwork::Serialize(FArchive& Archive)
 #endif // WITH_EDITORONLY_DATA
 	Super::Serialize(Archive);
 }
+
+float UNeuralNetwork::GetLastInferenceTime() {
+	return ComputeStatsModule.GetLastSample();
+};
+
+FNNIStatsData UNeuralNetwork::GetInferenceStats() {
+	return ComputeStatsModule.GetStats();
+};
+
+FNNIStatsData UNeuralNetwork::GetInputMemoryTransferStats() {
+	return InputMemoryTransferStatsModule.GetStats();
+};
diff --git a/Engine/Plugins/Experimental/NeuralNetworkInference/Source/NeuralNetworkInference/Public/NeuralNetwork.h b/Engine/Plugins/Experimental/NeuralNetworkInference/Source/NeuralNetworkInference/Public/NeuralNetwork.h
--- a/Engine/Plugins/Experimental/NeuralNetworkInference/Source/NeuralNetworkInference/Public/NeuralNetwork.h
+++ b/Engine/Plugins/Experimental/NeuralNetworkInference/Source/NeuralNetworkInference/Public/NeuralNetwork.h
@@ -5,6 +5,7 @@
 #include "CoreMinimal.h"
 #include "NeuralEnumClasses.h"
 #include "NeuralTensor.h"
+#include "NeuralStats.h"
 #include "NeuralNetwork.generated.h"
 
 /**
@@ -175,6 +176,17 @@ public:
 	 */
 	void Run();
 
+	/**
+	 * Stats functions: 
+	 * - GetLastInferenceTime will provide the last inference time measured milliseconds
+	 * - GetInferenceStats, returns Inference time statistics. (NumberSamples, Average, StdDev, Min, Max statistics measured in milliseconds)
+	 * - GetInputMemoryTransferStats, returns Input Memory Transfer statistics. (NumberSamples, Average, StdDev, Min, Max statistics measured in milliseconds)
+	 */
+	float GetLastInferenceTime();
+	FNNIStatsData GetInferenceStats();
+	FNNIStatsData GetInputMemoryTransferStats();
+	void ResetStats();
+
 protected:
 	/**
 	 * Whether Run() will use CPU or GPU acceleration hardware.
@@ -263,6 +275,9 @@ private:
 	FNeuralTensor& GetInputTensorMutable(const int32 InTensorIndex = 0);
 	FNeuralTensor& GetOutputTensorMutable(const int32 InTensorIndex = 0);
 
+	FNeuralStats ComputeStatsModule;
+	FNeuralStats InputMemoryTransferStatsModule;
+
 public:
 	/**
 	 * Internal function not needed by the user.
diff --git a/Engine/Plugins/Experimental/NeuralNetworkInference/Source/NeuralNetworkInference/Public/NeuralNetworkInferenceTimer.h b/Engine/Plugins/Experimental/NeuralNetworkInference/Source/NeuralNetworkInference/Public/NeuralNetworkInferenceTimer.h
deleted file mode 100644
diff --git a/Engine/Plugins/Experimental/NeuralNetworkInference/Source/NeuralNetworkInferenceProfiling/NeuralNetworkInferenceProfiling.Build.cs b/Engine/Plugins/Experimental/NeuralNetworkInference/Source/NeuralNetworkInferenceProfiling/NeuralNetworkInferenceProfiling.Build.cs
new file mode 100644
--- /dev/null
+++ b/Engine/Plugins/Experimental/NeuralNetworkInference/Source/NeuralNetworkInferenceProfiling/NeuralNetworkInferenceProfiling.Build.cs
@@ -0,0 +1,38 @@
+// Copyright Epic Games, Inc. All Rights Reserved.
+
+using UnrealBuildTool;
+using System.IO;
+
+public class NeuralNetworkInferenceProfiling : ModuleRules
+{
+	public NeuralNetworkInferenceProfiling(ReadOnlyTargetRules Target) : base(Target)
+	{
+
+		ShortName = "NNIProfiling"; 
+		PCHUsage = ModuleRules.PCHUsageMode.UseExplicitOrSharedPCHs;
+
+		PublicIncludePaths.AddRange(
+			new string[] {
+				System.IO.Path.Combine(ModuleDirectory, "..")
+			}
+		);
+
+		PublicDependencyModuleNames.AddRange
+			(
+			new string[] {
+				"Core",
+				"Engine",
+			}
+		);
+
+		PrivateDependencyModuleNames.AddRange
+			(
+			new string[] {
+				"ThirdPartyHelperAndDLLLoader",
+				"Boost"
+			}
+		);
+
+		PublicDefinitions.Add("WITH_NNI_STATS");
+	}
+}
diff --git a/Engine/Plugins/Experimental/NeuralNetworkInference/Source/NeuralNetworkInferenceProfiling/Private/NeuralNetworkInferenceProfilingModule.cpp b/Engine/Plugins/Experimental/NeuralNetworkInference/Source/NeuralNetworkInferenceProfiling/Private/NeuralNetworkInferenceProfilingModule.cpp
new file mode 100644
--- /dev/null
+++ b/Engine/Plugins/Experimental/NeuralNetworkInference/Source/NeuralNetworkInferenceProfiling/Private/NeuralNetworkInferenceProfilingModule.cpp
@@ -0,0 +1,16 @@
+// Copyright Epic Games, Inc. All Rights Reserved.
+
+#include "NeuralNetworkInferenceProfilingModule.h"
+
+// This code will execute after your module is loaded into memory; the exact timing is specified in the .uplugin file per-module
+void FNeuralNetworkInferenceProfilingModule::StartupModule()
+{
+}
+
+// This function may be called during shutdown to clean up your module. For modules that support dynamic reloading,
+// we call this function before unloading the module.
+void FNeuralNetworkInferenceProfilingModule::ShutdownModule()
+{
+}
+
+IMPLEMENT_MODULE(FNeuralNetworkInferenceProfilingModule, NeuralNetworkInferenceProfiling);
diff --git a/Engine/Plugins/Experimental/NeuralNetworkInference/Source/NeuralNetworkInferenceProfiling/Private/NeuralNetworkInferenceProfilingModule.h b/Engine/Plugins/Experimental/NeuralNetworkInference/Source/NeuralNetworkInferenceProfiling/Private/NeuralNetworkInferenceProfilingModule.h
new file mode 100644
--- /dev/null
+++ b/Engine/Plugins/Experimental/NeuralNetworkInference/Source/NeuralNetworkInferenceProfiling/Private/NeuralNetworkInferenceProfilingModule.h
@@ -0,0 +1,13 @@
+// Copyright Epic Games, Inc. All Rights Reserved.
+
+#pragma once
+
+#include "Modules/ModuleManager.h"
+
+class FNeuralNetworkInferenceProfilingModule : public IModuleInterface
+{
+public:
+	// Begin IModuleInterface
+	virtual void StartupModule() override;
+	virtual void ShutdownModule() override;
+};
diff --git a/Engine/Plugins/Experimental/NeuralNetworkInference/Source/NeuralNetworkInferenceProfiling/Private/NeuralNetworkInferenceTimer.cpp b/Engine/Plugins/Experimental/NeuralNetworkInference/Source/NeuralNetworkInferenceProfiling/Private/NeuralNetworkInferenceTimer.cpp
new file mode 100644
--- /dev/null
+++ b/Engine/Plugins/Experimental/NeuralNetworkInference/Source/NeuralNetworkInferenceProfiling/Private/NeuralNetworkInferenceTimer.cpp
@@ -0,0 +1,17 @@
+// Copyright Epic Games, Inc. All Rights Reserved.
+
+#include "NeuralNetworkInferenceTimer.h"
+#include "HAL/PlatformTime.h"
+
+/* FNeuralNetworkInferenceQATimer public functions
+ *****************************************************************************/
+
+void FNeuralNetworkInferenceTimer::Tic()
+{
+	TimeStart = FPlatformTime::Seconds();
+}
+
+float FNeuralNetworkInferenceTimer::Toc() const
+{
+	return (FPlatformTime::Seconds() - TimeStart) * 1e3;
+}
diff --git a/Engine/Plugins/Experimental/NeuralNetworkInference/Source/NeuralNetworkInferenceProfiling/Private/NeuralStats.cpp b/Engine/Plugins/Experimental/NeuralNetworkInference/Source/NeuralNetworkInferenceProfiling/Private/NeuralStats.cpp
new file mode 100644
--- /dev/null
+++ b/Engine/Plugins/Experimental/NeuralNetworkInference/Source/NeuralNetworkInferenceProfiling/Private/NeuralStats.cpp
@@ -0,0 +1,114 @@
+// Copyright Epic Games, Inc. All Rights Reserved.
+
+#include "NeuralStats.h"
+#include "GenericPlatform/GenericPlatformMath.h"
+DEFINE_LOG_CATEGORY_STATIC(LogNNIStats, Display, All);
+
+
+FNeuralStats::FNeuralStats(const uint32 InSizeRollingWindow) 
+: DataBuffer(InSizeRollingWindow, 0)
+{
+	SizeRollingWindow = InSizeRollingWindow;
+	BufferIdx = 0;
+	bIsBufferFull = false;
+}
+	
+void FNeuralStats::StoreSample(const float InRunTime) 
+{
+	LastSample = InRunTime;
+
+	DataBuffer[BufferIdx] = LastSample;
+
+	if(BufferIdx == DataBuffer.Capacity() && !bIsBufferFull)
+	{
+		bIsBufferFull = true;
+	}
+	BufferIdx = DataBuffer.GetNextIndex(BufferIdx);
+
+}
+
+float FNeuralStats::GetLastSample()
+{
+	return LastSample;
+}
+
+FNNIStatsData FNeuralStats::GetStats()
+{
+	const uint32 count = (bIsBufferFull) ? DataBuffer.Capacity() : BufferIdx;
+	const float Mean = CalculateMean();
+	const float StdDev = CalculateStdDev(Mean);
+	FNNIStatsMinMax MinMaxData = CalculateMinMax();
+
+	FNNIStatsData AggregatedData(count,	Mean, StdDev, MinMaxData.Min, MinMaxData.Max);
+
+	return AggregatedData;
+}
+
+void FNeuralStats::SetSizeRollingWindow(const uint32 InSizeRollingWindow) 
+{
+	DataBuffer = TCircularBuffer<float>(InSizeRollingWindow, 0);
+	ResetStats();
+}
+
+float FNeuralStats::CalculateMean() 
+{
+	uint32 PtrEnd = (bIsBufferFull) ? DataBuffer.Capacity(): BufferIdx;
+	float AccumulatorMean = 0;
+
+	for (uint32 Index = 0; Index < PtrEnd; ++Index)
+	{
+		AccumulatorMean += DataBuffer[Index];
+	}
+
+	const float Mean = AccumulatorMean / (static_cast<float>(PtrEnd) + EpsilonFloat);
+
+	return Mean;
+}
+
+float FNeuralStats::CalculateStdDev(const float InMean)
+{
+	uint32 PtrEnd = (bIsBufferFull) ? DataBuffer.Capacity() : BufferIdx;
+
+	float AccumulatorStdDev = 0;
+
+	for (uint32 Index = 0; Index < PtrEnd; ++Index)
+	{
+		AccumulatorStdDev += (DataBuffer[Index] - InMean) * (DataBuffer[Index] - InMean);
+	}
+
+	const float StdDev = FMath::Sqrt(AccumulatorStdDev / (static_cast<float>(PtrEnd) + EpsilonFloat));
+
+	return StdDev;
+}
+
+FNNIStatsMinMax FNeuralStats::CalculateMinMax()
+{
+	FNNIStatsMinMax MinMaxData;
+	uint32 PtrEnd = (bIsBufferFull) ? DataBuffer.Capacity() : BufferIdx;
+
+	for (uint32 Index = 0; Index < PtrEnd; ++Index)
+	{
+		const float CurrentElement = DataBuffer[Index];
+		if( CurrentElement < MinMaxData.Min)
+		{
+			MinMaxData.Min = CurrentElement;
+		}
+		if (CurrentElement > MinMaxData.Max)
+		{
+			MinMaxData.Max = CurrentElement;
+		}
+	}
+
+	return MinMaxData;
+}
+
+void FNeuralStats::ResetStats()
+{
+	for (uint32 Index = 0; Index < DataBuffer.Capacity(); ++Index)
+	{
+		DataBuffer[Index] = 0;
+	}
+	bIsBufferFull = false;
+	BufferIdx = 0;
+	LastSample = 0;
+}
diff --git a/Engine/Plugins/Experimental/NeuralNetworkInference/Source/NeuralNetworkInferenceProfiling/Public/NeuralNetworkInferenceTimer.h b/Engine/Plugins/Experimental/NeuralNetworkInference/Source/NeuralNetworkInferenceProfiling/Public/NeuralNetworkInferenceTimer.h
new file mode 100644
--- /dev/null
+++ b/Engine/Plugins/Experimental/NeuralNetworkInference/Source/NeuralNetworkInferenceProfiling/Public/NeuralNetworkInferenceTimer.h
@@ -0,0 +1,19 @@
+// Copyright Epic Games, Inc. All Rights Reserved.
+
+#pragma once
+
+#include "CoreMinimal.h"
+
+class NEURALNETWORKINFERENCEPROFILING_API FNeuralNetworkInferenceTimer
+{
+public:
+	void Tic();
+
+	/**
+	 * Time in milliseconds, but with nanosecond accuracy (e.g., 1.234567 msec).
+	 */
+	float Toc() const;
+
+private:
+	double TimeStart;
+};
diff --git a/Engine/Plugins/Experimental/NeuralNetworkInference/Source/NeuralNetworkInferenceProfiling/Public/NeuralStats.h b/Engine/Plugins/Experimental/NeuralNetworkInference/Source/NeuralNetworkInferenceProfiling/Public/NeuralStats.h
new file mode 100644
--- /dev/null
+++ b/Engine/Plugins/Experimental/NeuralNetworkInference/Source/NeuralNetworkInferenceProfiling/Public/NeuralStats.h
@@ -0,0 +1,75 @@
+// Copyright Epic Games, Inc. All Rights Reserved.
+
+#pragma once
+
+#include "CoreMinimal.h"
+#include "Containers/CircularBuffer.h"
+
+
+struct NEURALNETWORKINFERENCEPROFILING_API FNNISampleStatData
+{
+	FNNISampleStatData() {};
+
+	FNNISampleStatData(const float InInferenceTime)
+		: InferenceTime(InInferenceTime)	
+	{}
+
+	float InferenceTime = 0;
+};
+
+struct NEURALNETWORKINFERENCEPROFILING_API FNNIStatsData
+{
+	FNNIStatsData() {}
+	
+	FNNIStatsData(const int InNumberSamples, const float InAverage, const float InStdDev, const float InMin, const float InMax)
+		: NumberSamples(InNumberSamples)
+		, Average(InAverage)
+		, StdDev(InStdDev)
+		, Min(InMin)
+		, Max(InMax)
+		{}
+
+	uint32 NumberSamples = 0;
+	float Average = 0;
+	float StdDev = 0;
+	float Min = 0;
+	float Max = 0;
+};
+
+struct NEURALNETWORKINFERENCEPROFILING_API FNNIStatsMinMax
+{
+	float Min = FLT_MAX;
+	float Max = FLT_MIN;
+};
+
+
+class NEURALNETWORKINFERENCEPROFILING_API FNeuralStats
+{
+
+public:
+
+	FNeuralStats(const uint32 InSizeRollingWindow = 1024);
+	
+	void StoreSample(const float InRunTime);
+	void ResetStats();
+	void SetSizeRollingWindow(const uint32 InSizeRollingWindow);
+
+	float GetLastSample();
+	FNNIStatsData GetStats();
+
+private:
+
+	uint32 BufferIdx = 0;
+	bool bIsBufferFull = false;
+	uint32 SizeRollingWindow = 1024;
+	const float EpsilonFloat = 0.000001;
+	float LastSample;
+	TCircularBuffer<float> DataBuffer;
+
+	float CalculateMean();
+	float CalculateStdDev(const float InMean);
+	FNNIStatsMinMax CalculateMinMax();
+	
+};
+
+
diff --git a/Engine/Plugins/Experimental/NeuralNetworkInference/Source/NeuralNetworkInferenceQA/NeuralNetworkInferenceQA.Build.cs b/Engine/Plugins/Experimental/NeuralNetworkInference/Source/NeuralNetworkInferenceQA/NeuralNetworkInferenceQA.Build.cs
--- a/Engine/Plugins/Experimental/NeuralNetworkInference/Source/NeuralNetworkInferenceQA/NeuralNetworkInferenceQA.Build.cs
+++ b/Engine/Plugins/Experimental/NeuralNetworkInference/Source/NeuralNetworkInferenceQA/NeuralNetworkInferenceQA.Build.cs
@@ -30,7 +30,8 @@ public class NeuralNetworkInferenceQA : ModuleRules
 			new string[] {
 				"ModelProto",
 				"ModelProtoFileReader",
-				"NeuralNetworkInferenceCore"
+				"NeuralNetworkInferenceCore",
+				"NeuralNetworkInferenceProfiling"
 			}
 		);
 	}
diff --git a/Engine/Plugins/Experimental/NeuralNetworkInference/Source/NeuralNetworkInferenceQA/Private/ModelUnitTester.cpp b/Engine/Plugins/Experimental/NeuralNetworkInference/Source/NeuralNetworkInferenceQA/Private/ModelUnitTester.cpp
--- a/Engine/Plugins/Experimental/NeuralNetworkInference/Source/NeuralNetworkInferenceQA/Private/ModelUnitTester.cpp
+++ b/Engine/Plugins/Experimental/NeuralNetworkInference/Source/NeuralNetworkInferenceQA/Private/ModelUnitTester.cpp
@@ -6,6 +6,13 @@
 #include "Misc/Paths.h"
 
 
+struct FNNIUnitTesterTimeData 
+{
+	FNNIStatsData ComputeTimeData;
+	FNNIStatsData InputCopyTimeData;
+	FNNIStatsData OutputCopyTimeData;
+};
+
 
 /* FModelUnitTester static public functions
  *****************************************************************************/
@@ -348,34 +355,54 @@ else
 	return bDidGlobalTestPassed;
 }
 
-#define MODEL_UNIT_TESTER_SPEED_TEST(OutCopyTimeInMilliSeconds, OutNetworkTimeInMilliSeconds, InRepetitions, bInIsGPU) \
-	/* Input/output copy speed */ \
-	Timer.Tic(); \
-	for (int32 TimerIndex = 0; TimerIndex < InRepetitions; ++TimerIndex) \
-	{ \
-		InOutNetwork->SetInputFromArrayCopy(InputArray); \
-		CPUGPUCPUOutput = InOutNetwork->GetOutputTensor().GetArrayCopy<float>(); \
-	} \
-	const float OutCopyTimeInMilliSeconds = Timer.Toc() / InRepetitions; \
-	/* Forward() speed */ \
-	if (InRepetitions > 1) \
-	{ \
-		for (int32 TimerIndex = 0; TimerIndex < 5; ++TimerIndex) \
-		{ \
-			InOutNetwork->Run(); \
-		} \
-	} \
-	Timer.Tic(); \
-	if (InRepetitions > 0) \
-	{ \
-		for (int32 TimerIndex = 0; TimerIndex < InRepetitions; ++TimerIndex) \
-		{ \
-			InOutNetwork->SetInputFromArrayCopy(InputArray); \
-			InOutNetwork->Run(); \
-			CPUGPUCPUOutput = InOutNetwork->GetOutputTensor().GetArrayCopy<float>(); \
-		} \
-	} \
-	const float OutNetworkTimeInMilliSeconds = Timer.Toc() / InRepetitions - OutCopyTimeInMilliSeconds
+FNNIUnitTesterTimeData GetTimeInformation(
+		const int32 InRepetitions,
+		FNeuralNetworkInferenceTimer& Timer,
+		UNeuralNetwork* InOutNetwork,
+		TArray<float>* InputArray,
+		TArray<float>* CPUGPUCPUOutput) 
+	{
+
+		FNeuralStats OutCopyingStats;
+
+		/* Input/output copy speed */ 
+		for (int32 TimerIndex = 0; TimerIndex < InRepetitions; ++TimerIndex) 
+		{ 
+			InOutNetwork->SetInputFromArrayCopy(*InputArray); 
+			Timer.Tic();
+			*CPUGPUCPUOutput = InOutNetwork->GetOutputTensor().GetArrayCopy<float>();
+			const float CurrentCopyingTime = Timer.Toc();
+			OutCopyingStats.StoreSample(CurrentCopyingTime);
+
+		} 
+
+		/* Forward() speed */ 
+		if (InRepetitions > 1) 
+		{ 
+			for (int32 TimerIndex = 0; TimerIndex < 5; ++TimerIndex) 
+			{ 
+				InOutNetwork->Run(); 
+			} 
+		} 
+
+		FNeuralStats CurrentComputingStats;
+		if (InRepetitions > 0) 
+		{ 
+			for (int32 TimerIndex = 0; TimerIndex < InRepetitions; ++TimerIndex) 
+			{ 
+				InOutNetwork->SetInputFromArrayCopy(*InputArray); 
+				InOutNetwork->Run();
+				*CPUGPUCPUOutput = InOutNetwork->GetOutputTensor().GetArrayCopy<float>(); 
+			} 
+		} 
+		FNNIUnitTesterTimeData NetworkTimeData;
+		NetworkTimeData.ComputeTimeData = InOutNetwork->GetInferenceStats();
+		NetworkTimeData.InputCopyTimeData = InOutNetwork->GetInputMemoryTransferStats();
+		NetworkTimeData.OutputCopyTimeData = OutCopyingStats.GetStats();
+
+		return NetworkTimeData;
+	}
+
 
 bool FModelUnitTester::ModelSpeedTest(const FString& InUAssetPath, const ENeuralDeviceType InDeviceType, const ENeuralBackEnd InBackEnd, const int32 InRepetitions)
 {
@@ -410,14 +437,20 @@ bool FModelUnitTester::ModelSpeedTest(const FString& InUAssetPath, const ENeural
 	FNeuralNetworkInferenceTimer Timer;
 	// Run profiling 1 time
 	InOutNetwork->SetDeviceType(InDeviceType);
-	MODEL_UNIT_TESTER_SPEED_TEST(CopyTimer1, NetworkTimer1, 1, /*bIsGPU*/false);
+	InOutNetwork->ResetStats();
+	FNNIUnitTesterTimeData TimeData1 = GetTimeInformation(1, Timer, InOutNetwork, &InputArray, &CPUGPUCPUOutput);
+
 	// Run profiling n times
 	InOutNetwork->SetDeviceType(InDeviceType);
-	MODEL_UNIT_TESTER_SPEED_TEST(CopyTimer, NetworkTimer, InRepetitions, /*bIsGPU*/false);
+	InOutNetwork->ResetStats();
+	FNNIUnitTesterTimeData TimeDataN = GetTimeInformation(InRepetitions, Timer, InOutNetwork, &InputArray, &CPUGPUCPUOutput);
+
 	// Display speed times
 	UE_LOG(LogNeuralNetworkInferenceQA, Display,
-		TEXT("%s-%s:\t1 time = %f+%f msec, avg(%d times) = %f+%f msec."),
-		*BackEndString, *DeviceTypeString, CopyTimer1, NetworkTimer1, InRepetitions, CopyTimer, NetworkTimer);
+		TEXT("%s-%s:, Times(msec)-> Inference = %f, Copy In = %f, Copy Out %f, avg(%d times) = Inference = %f, Copy In = %f, Copy Out %f."),
+		*BackEndString, *DeviceTypeString, TimeData1.ComputeTimeData.Average, TimeData1.InputCopyTimeData.Average, TimeData1.OutputCopyTimeData.Average,
+		InRepetitions, TimeDataN.ComputeTimeData.Average, TimeDataN.InputCopyTimeData.Average, TimeDataN.OutputCopyTimeData.Average);
+
 	// Reset to original network state
 	InOutNetwork->SetDeviceType(OriginalDeviceType);
 	InOutNetwork->SetBackEnd(OriginalBackEnd);
