Fixed horde warning on Win64 Clang: You have included windows.h before MinWindows.h [-W#pragma-messages]
#rnx

diff --git a/Engine/Plugins/Experimental/NeuralNetworkInference/Source/NeuralNetworkInference/Private/NeuralNetworkImplBackEndUEAndORT.cpp b/Engine/Plugins/Experimental/NeuralNetworkInference/Source/NeuralNetworkInference/Private/NeuralNetworkImplBackEndUEAndORT.cpp
--- a/Engine/Plugins/Experimental/NeuralNetworkInference/Source/NeuralNetworkInference/Private/NeuralNetworkImplBackEndUEAndORT.cpp
+++ b/Engine/Plugins/Experimental/NeuralNetworkInference/Source/NeuralNetworkInference/Private/NeuralNetworkImplBackEndUEAndORT.cpp
@@ -16,9 +16,12 @@
 	// Disable NOMINMAX & WIN32_LEAN_AND_MEAN defines to avoid compiler warnings
 	#pragma push_macro("NOMINMAX")
 	#pragma push_macro("WIN32_LEAN_AND_MEAN")
+	#pragma push_macro("UE_MINIMAL_WINDOWS_INCLUDE")
 	#undef NOMINMAX
 	#undef WIN32_LEAN_AND_MEAN
+	#define UE_MINIMAL_WINDOWS_INCLUDE // Avoids Win64 Clang warning
 	#include "D3D12RHIPrivate.h"
+	#pragma pop_macro("UE_MINIMAL_WINDOWS_INCLUDE")
 	#pragma pop_macro("WIN32_LEAN_AND_MEAN")
 	#pragma pop_macro("NOMINMAX")
 
