Added "IsExperimentalVersion" to NNI.uplugin file

Thanks [REDACTED_HANDLE] for reporting

#rnx

diff --git a/Engine/Plugins/Experimental/NeuralNetworkInference/NeuralNetworkInference.uplugin b/Engine/Plugins/Experimental/NeuralNetworkInference/NeuralNetworkInference.uplugin
--- a/Engine/Plugins/Experimental/NeuralNetworkInference/NeuralNetworkInference.uplugin
+++ b/Engine/Plugins/Experimental/NeuralNetworkInference/NeuralNetworkInference.uplugin
@@ -13,6 +13,7 @@
 	"EnabledByDefault": false,
 	"CanContainContent": false,
 	"IsBetaVersion": true,
+	"IsExperimentalVersion": true,
 	"Installed": false,
 	"Modules": [
 		{
