Fixed CIS issue
#rnx

diff --git a/Engine/Plugins/Experimental/NeuralNetworkInference/Source/NeuralNetworkInference/Private/NeuralNetworkImplBackEndUEAndORT.cpp b/Engine/Plugins/Experimental/NeuralNetworkInference/Source/NeuralNetworkInference/Private/NeuralNetworkImplBackEndUEAndORT.cpp
--- a/Engine/Plugins/Experimental/NeuralNetworkInference/Source/NeuralNetworkInference/Private/NeuralNetworkImplBackEndUEAndORT.cpp
+++ b/Engine/Plugins/Experimental/NeuralNetworkInference/Source/NeuralNetworkInference/Private/NeuralNetworkImplBackEndUEAndORT.cpp
@@ -9,22 +9,14 @@
 #include "RedirectCoutAndCerrToUeLog.h"
 
 #if defined(WITH_UE_AND_ORT_SUPPORT) && defined(PLATFORM_WIN64)
-	#include "HAL/CriticalSection.h"
-	#include "RHI.h"
-	#include "DynamicRHI.h"
-	
 	// Disable NOMINMAX & WIN32_LEAN_AND_MEAN defines to avoid compiler warnings
 	#pragma push_macro("NOMINMAX")
 	#pragma push_macro("WIN32_LEAN_AND_MEAN")
-	#pragma push_macro("UE_MINIMAL_WINDOWS_INCLUDE")
 	#undef NOMINMAX
 	#undef WIN32_LEAN_AND_MEAN
-	#define UE_MINIMAL_WINDOWS_INCLUDE // Avoids Win64 Clang warning
 	#include "D3D12RHIPrivate.h"
-	#pragma pop_macro("UE_MINIMAL_WINDOWS_INCLUDE")
 	#pragma pop_macro("WIN32_LEAN_AND_MEAN")
 	#pragma pop_macro("NOMINMAX")
-
 #endif
 
 //#define WITH_NNI_CPU_NOT_RECOMMENDED // Only for debugging purposes
diff --git a/Engine/Plugins/Experimental/NeuralNetworkInference/Source/NeuralNetworkInference/Private/NeuralNetworkImplBackEndUEAndORT.h b/Engine/Plugins/Experimental/NeuralNetworkInference/Source/NeuralNetworkInference/Private/NeuralNetworkImplBackEndUEAndORT.h
--- a/Engine/Plugins/Experimental/NeuralNetworkInference/Source/NeuralNetworkInference/Private/NeuralNetworkImplBackEndUEAndORT.h
+++ b/Engine/Plugins/Experimental/NeuralNetworkInference/Source/NeuralNetworkInference/Private/NeuralNetworkImplBackEndUEAndORT.h
@@ -4,6 +4,11 @@
 
 #include "NeuralNetwork.h"
 
+// RHI includes must happen before onnxruntime_cxx_api.h (both include Windows.h)
+#include "HAL/CriticalSection.h"
+#include "RHI.h"
+#include "DynamicRHI.h"
+
 #include "ThirdPartyWarningDisabler.h"
 NNI_THIRD_PARTY_INCLUDES_START
 #undef check
