Proper shutdown - Part 1/2
- GC not crashing NN on destruction / Editor shutdown. Low level: Added IsReadyForFinishDestroy() override to NN as well as std::atomic<bool> bIsRunningNetwork that will make sure the network is not destroyed while running.
- Only working for synchronous mode (which is ok because async mode is not implemented yet anyway).

#rnx

diff --git a/Engine/Plugins/Experimental/NeuralNetworkInference/Source/NeuralNetworkInference/Private/NeuralNetwork.cpp b/Engine/Plugins/Experimental/NeuralNetworkInference/Source/NeuralNetworkInference/Private/NeuralNetwork.cpp
--- a/Engine/Plugins/Experimental/NeuralNetworkInference/Source/NeuralNetworkInference/Private/NeuralNetwork.cpp
+++ b/Engine/Plugins/Experimental/NeuralNetworkInference/Source/NeuralNetworkInference/Private/NeuralNetwork.cpp
@@ -51,6 +51,7 @@ UNeuralNetwork::UNeuralNetwork()
 	, SynchronousMode(ENeuralNetworkSynchronousMode::Synchronous)
 	, BackEnd(ENeuralBackEnd::Auto)
 	, bIsLoaded(false)
+	, bIsRunningNetwork(false)
 	, BackEndForCurrentPlatform(FPrivateNeuralNetwork::SetBackEndForCurrentPlatform(BackEnd))
 {
 }
@@ -263,22 +264,18 @@ void UNeuralNetwork::SetInputFromArrayCopy(const TArray<float>& InArray, const i
 	{
 		ImplBackEndUEAndORT->InputTensors[InTensorIndex].SetFromArrayCopy(InArray);
 	}
-
 	// UEOnly
 	else if (BackEndForCurrentPlatform == ENeuralBackEnd::UEOnly)
 	{
 		FNeuralTensorManager& TensorManager = ImplBackEndUEOnly->TensorManager;
 		TensorManager.GetTensorsMutable()[TensorManager.GetInputIndexes()[InTensorIndex]].SetFromArrayCopy(InArray);
 	}
-
 	// Unknown
 	else
 	{
 		UE_LOG(LogNeuralNetworkInference, Warning, TEXT("UNeuralNetwork::SetInputFromArrayCopy(): Unknown [BackEnd,BackEndForCurrentPlatform] = [%d,%d]."), (int32)BackEnd, (int32)BackEndForCurrentPlatform);
 	}
-
 	InputMemoryTransferStatsModule.StoreSample(RunTimer.Toc());
-
 }
 
 void* UNeuralNetwork::GetInputDataPointerMutable(const int32 InTensorIndex)
@@ -490,6 +487,8 @@ void UNeuralNetwork::OutputTensorsToCPU(const TArray<int32>& InTensorIndexes)
 
 void UNeuralNetwork::Run()
 {
+	bIsRunningNetwork = true;
+
 	DECLARE_SCOPE_CYCLE_COUNTER(TEXT("UNeuralNetwork_Run"), STAT_UNeuralNetwork_Run, STATGROUP_MachineLearning);
 
 	// Sanity check
@@ -505,22 +504,33 @@ void UNeuralNetwork::Run()
 	{
 		ImplBackEndUEAndORT->Run(SynchronousMode, InputDeviceType, OutputDeviceType);
 	}
-
 	// UEOnly
 	else if (BackEndForCurrentPlatform == ENeuralBackEnd::UEOnly)
 	{
 		ImplBackEndUEOnly->Run(OnAsyncRunCompletedDelegate, SynchronousMode, DeviceType, InputDeviceType, OutputDeviceType);
 	}
-
 	// Unknown
 	else
 	{
 		UE_LOG(LogNeuralNetworkInference, Warning, TEXT("UNeuralNetwork::Run(): Unknown [BackEnd,BackEndForCurrentPlatform] = [%d,%d]."), (int32)BackEnd, (int32)BackEndForCurrentPlatform);
 	}
-
 	ComputeStatsModule.StoreSample(RunTimer.Toc());
+
+	bIsRunningNetwork = false;
 }
 
+float UNeuralNetwork::GetLastInferenceTime() {
+	return ComputeStatsModule.GetLastSample();
+};
+
+FNNIStatsData UNeuralNetwork::GetInferenceStats() {
+	return ComputeStatsModule.GetStats();
+};
+
+FNNIStatsData UNeuralNetwork::GetInputMemoryTransferStats() {
+	return InputMemoryTransferStatsModule.GetStats();
+};
+
 
 /* UNeuralNetwork private functions
  *****************************************************************************/
@@ -657,14 +667,7 @@ void UNeuralNetwork::Serialize(FArchive& Archive)
 	Super::Serialize(Archive);
 }
 
-float UNeuralNetwork::GetLastInferenceTime() {
-	return ComputeStatsModule.GetLastSample();
-};
-
-FNNIStatsData UNeuralNetwork::GetInferenceStats() {
-	return ComputeStatsModule.GetStats();
-};
-
-FNNIStatsData UNeuralNetwork::GetInputMemoryTransferStats() {
-	return InputMemoryTransferStatsModule.GetStats();
-};
+bool UNeuralNetwork::IsReadyForFinishDestroy()
+{
+	return !bIsRunningNetwork;
+}
diff --git a/Engine/Plugins/Experimental/NeuralNetworkInference/Source/NeuralNetworkInference/Public/NeuralNetwork.h b/Engine/Plugins/Experimental/NeuralNetworkInference/Source/NeuralNetworkInference/Public/NeuralNetwork.h
--- a/Engine/Plugins/Experimental/NeuralNetworkInference/Source/NeuralNetworkInference/Public/NeuralNetwork.h
+++ b/Engine/Plugins/Experimental/NeuralNetworkInference/Source/NeuralNetworkInference/Public/NeuralNetwork.h
@@ -229,6 +229,13 @@ protected:
 private:
 	bool bIsLoaded;
 
+	/**
+	 * Whether an inference pass (i.e., Run) is happening.
+	 * This variable is thread safe as long as only "Run" modifies it. Other functions can safely read it at any time.
+	 * If other functions outside of Run() have to modify it, consider using mutexes.
+	 */
+	std::atomic<bool> bIsRunningNetwork;
+
 	UPROPERTY()
 	TArray<uint8> ModelReadFromFileInBytes;
 
@@ -310,5 +317,6 @@ private:
 	virtual void PostInitProperties() override;
 	virtual void PostLoad() override;
 	virtual void Serialize(FArchive& Archive) override;
+	virtual bool IsReadyForFinishDestroy() override;
 	//~End of UObject interface
 };
