Fixed horde warning on Win64 Clang: You have included windows.h before MinWindows.h [-W#pragma-messages]

#rnx

diff --git a/Engine/Plugins/Experimental/NeuralNetworkInference/Source/NeuralNetworkInference/Private/NeuralNetworkImplBackEndUEAndORT.cpp b/Engine/Plugins/Experimental/NeuralNetworkInference/Source/NeuralNetworkInference/Private/NeuralNetworkImplBackEndUEAndORT.cpp
--- a/Engine/Plugins/Experimental/NeuralNetworkInference/Source/NeuralNetworkInference/Private/NeuralNetworkImplBackEndUEAndORT.cpp
+++ b/Engine/Plugins/Experimental/NeuralNetworkInference/Source/NeuralNetworkInference/Private/NeuralNetworkImplBackEndUEAndORT.cpp
@@ -1,23 +1,13 @@
 // Copyright Epic Games, Inc. All Rights Reserved.
 
 #include "NeuralNetworkImplBackEndUEAndORT.h"
+
 #include "NeuralNetworkInferenceUtils.h"
 #include "NeuralNetworkInferenceUtilsGPU.h"
+#include "RedirectCoutAndCerrToUeLog.h"
 #if WITH_EDITOR
 #include "Misc/MessageDialog.h"
 #endif //WITH_EDITOR
-#include "RedirectCoutAndCerrToUeLog.h"
-
-#if defined(WITH_UE_AND_ORT_SUPPORT) && defined(PLATFORM_WIN64)
-	// Disable NOMINMAX & WIN32_LEAN_AND_MEAN defines to avoid compiler warnings
-	#pragma push_macro("NOMINMAX")
-	#pragma push_macro("WIN32_LEAN_AND_MEAN")
-	#undef NOMINMAX
-	#undef WIN32_LEAN_AND_MEAN
-	#include "D3D12RHIPrivate.h"
-	#pragma pop_macro("WIN32_LEAN_AND_MEAN")
-	#pragma pop_macro("NOMINMAX")
-#endif
 
 //#define WITH_NNI_CPU_NOT_RECOMMENDED // Only for debugging purposes
 
diff --git a/Engine/Plugins/Experimental/NeuralNetworkInference/Source/NeuralNetworkInference/Private/NeuralNetworkImplBackEndUEAndORT.h b/Engine/Plugins/Experimental/NeuralNetworkInference/Source/NeuralNetworkInference/Private/NeuralNetworkImplBackEndUEAndORT.h
--- a/Engine/Plugins/Experimental/NeuralNetworkInference/Source/NeuralNetworkInference/Private/NeuralNetworkImplBackEndUEAndORT.h
+++ b/Engine/Plugins/Experimental/NeuralNetworkInference/Source/NeuralNetworkInference/Private/NeuralNetworkImplBackEndUEAndORT.h
@@ -8,6 +8,16 @@
 #include "HAL/CriticalSection.h"
 #include "RHI.h"
 #include "DynamicRHI.h"
+#if defined(WITH_UE_AND_ORT_SUPPORT) && defined(PLATFORM_WIN64)
+	// Disable NOMINMAX & WIN32_LEAN_AND_MEAN defines to avoid compiler warnings
+	#pragma push_macro("NOMINMAX")
+	#pragma push_macro("WIN32_LEAN_AND_MEAN")
+	#undef NOMINMAX
+	#undef WIN32_LEAN_AND_MEAN
+	#include "D3D12RHIPrivate.h"
+	#pragma pop_macro("WIN32_LEAN_AND_MEAN")
+	#pragma pop_macro("NOMINMAX")
+#endif
 
 #include "ThirdPartyWarningDisabler.h"
 NNI_THIRD_PARTY_INCLUDES_START
