NeuralNetworkInference (NNI) plugin: Making sure the "asynchronous" version of NNI (non-ORT) do not crash on destruction if a network was run right before shutting it down.

[REDACTED_HANDLE], this bool should be implemented by your code for ORT-backend as well, so you should understand this change.

#rnx

diff --git a/Engine/Plugins/Experimental/NeuralNetworkInference/Source/NeuralNetworkInference/Private/NeuralNetwork.cpp b/Engine/Plugins/Experimental/NeuralNetworkInference/Source/NeuralNetworkInference/Private/NeuralNetwork.cpp
--- a/Engine/Plugins/Experimental/NeuralNetworkInference/Source/NeuralNetworkInference/Private/NeuralNetwork.cpp
+++ b/Engine/Plugins/Experimental/NeuralNetworkInference/Source/NeuralNetworkInference/Private/NeuralNetwork.cpp
@@ -51,7 +51,7 @@ UNeuralNetwork::UNeuralNetwork()
 	, SynchronousMode(ENeuralNetworkSynchronousMode::Synchronous)
 	, BackEnd(ENeuralBackEnd::Auto)
 	, bIsLoaded(false)
-	, bIsRunningNetwork(false)
+	, bIsBackgroundThreadRunning(false)
 	, BackEndForCurrentPlatform(FPrivateNeuralNetwork::SetBackEndForCurrentPlatform(BackEnd))
 {
 }
@@ -487,8 +487,6 @@ void UNeuralNetwork::OutputTensorsToCPU(const TArray<int32>& InTensorIndexes)
 
 void UNeuralNetwork::Run()
 {
-	bIsRunningNetwork = true;
-
 	DECLARE_SCOPE_CYCLE_COUNTER(TEXT("UNeuralNetwork_Run"), STAT_UNeuralNetwork_Run, STATGROUP_MachineLearning);
 
 	// Sanity check
@@ -507,7 +505,7 @@ void UNeuralNetwork::Run()
 	// UEOnly
 	else if (BackEndForCurrentPlatform == ENeuralBackEnd::UEOnly)
 	{
-		ImplBackEndUEOnly->Run(OnAsyncRunCompletedDelegate, SynchronousMode, DeviceType, InputDeviceType, OutputDeviceType);
+		ImplBackEndUEOnly->Run(OnAsyncRunCompletedDelegate, bIsBackgroundThreadRunning, SynchronousMode, DeviceType, InputDeviceType, OutputDeviceType);
 	}
 	// Unknown
 	else
@@ -515,8 +513,6 @@ void UNeuralNetwork::Run()
 		UE_LOG(LogNeuralNetworkInference, Warning, TEXT("UNeuralNetwork::Run(): Unknown [BackEnd,BackEndForCurrentPlatform] = [%d,%d]."), (int32)BackEnd, (int32)BackEndForCurrentPlatform);
 	}
 	ComputeStatsModule.StoreSample(RunTimer.Toc());
-
-	bIsRunningNetwork = false;
 }
 
 float UNeuralNetwork::GetLastInferenceTime() {
@@ -669,5 +665,5 @@ void UNeuralNetwork::Serialize(FArchive& Archive)
 
 bool UNeuralNetwork::IsReadyForFinishDestroy()
 {
-	return !bIsRunningNetwork;
+	return !bIsBackgroundThreadRunning;
 }
diff --git a/Engine/Plugins/Experimental/NeuralNetworkInference/Source/NeuralNetworkInference/Private/NeuralNetworkImplBackEndUEOnly.cpp b/Engine/Plugins/Experimental/NeuralNetworkInference/Source/NeuralNetworkInference/Private/NeuralNetworkImplBackEndUEOnly.cpp
--- a/Engine/Plugins/Experimental/NeuralNetworkInference/Source/NeuralNetworkInference/Private/NeuralNetworkImplBackEndUEOnly.cpp
+++ b/Engine/Plugins/Experimental/NeuralNetworkInference/Source/NeuralNetworkInference/Private/NeuralNetworkImplBackEndUEOnly.cpp
@@ -74,7 +74,7 @@ bool UNeuralNetwork::FImplBackEndUEOnly::Load(TSharedPtr<FImplBackEndUEOnly>& In
 //	return (InOutImplBackEndUEOnly->Operators.Num() > 0 && InOutImplBackEndUEOnly->TensorManager.IsLoaded());
 //}
 
-void UNeuralNetwork::FImplBackEndUEOnly::Run(FOnAsyncRunCompleted& InOutOnAsyncRunCompletedDelegate, const ENeuralNetworkSynchronousMode InSynchronousMode, const ENeuralDeviceType InDeviceType, const ENeuralDeviceType InInputDeviceType, const ENeuralDeviceType InOutputDeviceType)
+void UNeuralNetwork::FImplBackEndUEOnly::Run(FOnAsyncRunCompleted& InOutOnAsyncRunCompletedDelegate, std::atomic<bool>& bIsBackgroundThreadRunning, const ENeuralNetworkSynchronousMode InSynchronousMode, const ENeuralDeviceType InDeviceType, const ENeuralDeviceType InInputDeviceType, const ENeuralDeviceType InOutputDeviceType)
 {
 	// Run UNeuralNetwork::UEOnly
 	if (Operators.Num() > 0)
@@ -106,7 +106,7 @@ void UNeuralNetwork::FImplBackEndUEOnly::Run(FOnAsyncRunCompleted& InOutOnAsyncR
 
 			// On RHI thread
 			ENQUEUE_RENDER_COMMAND(UNeuralNetwork_UEOnly_Run_RenderThread)(
-				[this, &InOutOnAsyncRunCompletedDelegate, InSynchronousMode, InInputDeviceType, InOutputDeviceType](FRHICommandListImmediate& RHICmdList)
+				[this, &InOutOnAsyncRunCompletedDelegate, &bIsBackgroundThreadRunning, InSynchronousMode, InInputDeviceType, InOutputDeviceType](FRHICommandListImmediate& RHICmdList)
 				{
 					FRDGBuilder GraphBuilder(RHICmdList, RDG_EVENT_NAME("FImplBackEndUEOnly::Run()"));
 
@@ -196,12 +196,14 @@ void UNeuralNetwork::FImplBackEndUEOnly::Run(FOnAsyncRunCompleted& InOutOnAsyncR
 					// Broadcast delegates (from the render thread)
 					if (InSynchronousMode == ENeuralNetworkSynchronousMode::Asynchronous)
 					{
+						bIsBackgroundThreadRunning = true; // This will be done right away
 						GraphBuilder.AddPass(
 							RDG_EVENT_NAME("Async delegate broadcast"),
 							ERDGPassFlags::None,
-							[&InOutOnAsyncRunCompletedDelegate](FRHICommandListImmediate& RHICmdList)
+							[&InOutOnAsyncRunCompletedDelegate, &bIsBackgroundThreadRunning](FRHICommandListImmediate& RHICmdList)
 						{
 							InOutOnAsyncRunCompletedDelegate.ExecuteIfBound();
+							bIsBackgroundThreadRunning = false; // This will be done in the render graph after all NNI operators have been run
 						});
 					}
 
diff --git a/Engine/Plugins/Experimental/NeuralNetworkInference/Source/NeuralNetworkInference/Private/NeuralNetworkImplBackEndUEOnly.h b/Engine/Plugins/Experimental/NeuralNetworkInference/Source/NeuralNetworkInference/Private/NeuralNetworkImplBackEndUEOnly.h
--- a/Engine/Plugins/Experimental/NeuralNetworkInference/Source/NeuralNetworkInference/Private/NeuralNetworkImplBackEndUEOnly.h
+++ b/Engine/Plugins/Experimental/NeuralNetworkInference/Source/NeuralNetworkInference/Private/NeuralNetworkImplBackEndUEOnly.h
@@ -36,7 +36,7 @@ struct UNeuralNetwork::FImplBackEndUEOnly
 
 	//static bool Load(TSharedPtr<FImplBackEndUEOnly>& InOutImplBackEndUEOnly, FNeuralTensorManager& InTensorManager, const TArray<TSharedPtr<FNeuralOperator>>& InOperators);
 
-	void Run(FOnAsyncRunCompleted& InOutOnAsyncRunCompletedDelegate, const ENeuralNetworkSynchronousMode InSynchronousMode, const ENeuralDeviceType InDeviceType, const ENeuralDeviceType InInputDeviceType, const ENeuralDeviceType InOutputDeviceType);
+	void Run(FOnAsyncRunCompleted& InOutOnAsyncRunCompletedDelegate, std::atomic<bool>& bIsBackgroundThreadRunning, const ENeuralNetworkSynchronousMode InSynchronousMode, const ENeuralDeviceType InDeviceType, const ENeuralDeviceType InInputDeviceType, const ENeuralDeviceType InOutputDeviceType);
 
 private:
 	/**
diff --git a/Engine/Plugins/Experimental/NeuralNetworkInference/Source/NeuralNetworkInference/Public/NeuralNetwork.h b/Engine/Plugins/Experimental/NeuralNetworkInference/Source/NeuralNetworkInference/Public/NeuralNetwork.h
--- a/Engine/Plugins/Experimental/NeuralNetworkInference/Source/NeuralNetworkInference/Public/NeuralNetwork.h
+++ b/Engine/Plugins/Experimental/NeuralNetworkInference/Source/NeuralNetworkInference/Public/NeuralNetwork.h
@@ -234,7 +234,7 @@ private:
 	 * This variable is thread safe as long as only "Run" modifies it. Other functions can safely read it at any time.
 	 * If other functions outside of Run() have to modify it, consider using mutexes.
 	 */
-	std::atomic<bool> bIsRunningNetwork;
+	std::atomic<bool> bIsBackgroundThreadRunning;
 
 	UPROPERTY()
 	TArray<uint8> ModelReadFromFileInBytes;
