Quick fix

_codereview Francisco.Vicente
#rnx

diff --git a/Engine/Plugins/Experimental/NeuralNetworkInference/Source/NeuralNetworkInference/Private/NeuralNetworkImplBackEndUEOnly.cpp b/Engine/Plugins/Experimental/NeuralNetworkInference/Source/NeuralNetworkInference/Private/NeuralNetworkImplBackEndUEOnly.cpp
--- a/Engine/Plugins/Experimental/NeuralNetworkInference/Source/NeuralNetworkInference/Private/NeuralNetworkImplBackEndUEOnly.cpp
+++ b/Engine/Plugins/Experimental/NeuralNetworkInference/Source/NeuralNetworkInference/Private/NeuralNetworkImplBackEndUEOnly.cpp
@@ -74,7 +74,7 @@ bool UNeuralNetwork::FImplBackEndUEOnly::Load(TSharedPtr<FImplBackEndUEOnly>& In
 //	return (InOutImplBackEndUEOnly->Operators.Num() > 0 && InOutImplBackEndUEOnly->TensorManager.IsLoaded());
 //}
 
-void UNeuralNetwork::FImplBackEndUEOnly::Run(FOnAsyncRunCompleted& InOutOnAsyncRunCompletedDelegate, std::atomic<bool>& bIsBackgroundThreadRunning, const ENeuralNetworkSynchronousMode InSynchronousMode, const ENeuralDeviceType InDeviceType, const ENeuralDeviceType InInputDeviceType, const ENeuralDeviceType InOutputDeviceType)
+void UNeuralNetwork::FImplBackEndUEOnly::Run(FOnAsyncRunCompleted& InOutOnAsyncRunCompletedDelegate, std::atomic<bool>& bInIsBackgroundThreadRunning, const ENeuralNetworkSynchronousMode InSynchronousMode, const ENeuralDeviceType InDeviceType, const ENeuralDeviceType InInputDeviceType, const ENeuralDeviceType InOutputDeviceType)
 {
 	// Run UNeuralNetwork::UEOnly
 	if (Operators.Num() > 0)
@@ -106,7 +106,7 @@ void UNeuralNetwork::FImplBackEndUEOnly::Run(FOnAsyncRunCompleted& InOutOnAsyncR
 
 			// On RHI thread
 			ENQUEUE_RENDER_COMMAND(UNeuralNetwork_UEOnly_Run_RenderThread)(
-				[this, &InOutOnAsyncRunCompletedDelegate, &bIsBackgroundThreadRunning, InSynchronousMode, InInputDeviceType, InOutputDeviceType](FRHICommandListImmediate& RHICmdList)
+				[this, &InOutOnAsyncRunCompletedDelegate, &bInIsBackgroundThreadRunning, InSynchronousMode, InInputDeviceType, InOutputDeviceType](FRHICommandListImmediate& RHICmdList)
 				{
 					FRDGBuilder GraphBuilder(RHICmdList, RDG_EVENT_NAME("FImplBackEndUEOnly::Run()"));
 
@@ -196,14 +196,14 @@ void UNeuralNetwork::FImplBackEndUEOnly::Run(FOnAsyncRunCompleted& InOutOnAsyncR
 					// Broadcast delegates (from the render thread)
 					if (InSynchronousMode == ENeuralNetworkSynchronousMode::Asynchronous)
 					{
-						bIsBackgroundThreadRunning = true; // This will be done right away
+						bInIsBackgroundThreadRunning = true; // This will be done right away
 						GraphBuilder.AddPass(
 							RDG_EVENT_NAME("Async delegate broadcast"),
 							ERDGPassFlags::None,
-							[&InOutOnAsyncRunCompletedDelegate, &bIsBackgroundThreadRunning](FRHICommandListImmediate& RHICmdList)
+							[&InOutOnAsyncRunCompletedDelegate, &bInIsBackgroundThreadRunning](FRHICommandListImmediate& RHICmdList)
 						{
 							InOutOnAsyncRunCompletedDelegate.ExecuteIfBound();
-							bIsBackgroundThreadRunning = false; // This will be done in the render graph after all NNI operators have been run
+							bInIsBackgroundThreadRunning = false; // This will be done in the render graph after all NNI operators have been run
 						});
 					}
 
diff --git a/Engine/Plugins/Experimental/NeuralNetworkInference/Source/NeuralNetworkInference/Private/NeuralNetworkImplBackEndUEOnly.h b/Engine/Plugins/Experimental/NeuralNetworkInference/Source/NeuralNetworkInference/Private/NeuralNetworkImplBackEndUEOnly.h
--- a/Engine/Plugins/Experimental/NeuralNetworkInference/Source/NeuralNetworkInference/Private/NeuralNetworkImplBackEndUEOnly.h
+++ b/Engine/Plugins/Experimental/NeuralNetworkInference/Source/NeuralNetworkInference/Private/NeuralNetworkImplBackEndUEOnly.h
@@ -4,6 +4,7 @@
 
 #include "NeuralNetwork.h"
 
+#include <atomic>
 #include "ModelProto.h"
 #include "NeuralOperator.h"
 #include "UEOnly/NeuralTensorManager.h"
@@ -36,7 +37,7 @@ struct UNeuralNetwork::FImplBackEndUEOnly
 
 	//static bool Load(TSharedPtr<FImplBackEndUEOnly>& InOutImplBackEndUEOnly, FNeuralTensorManager& InTensorManager, const TArray<TSharedPtr<FNeuralOperator>>& InOperators);
 
-	void Run(FOnAsyncRunCompleted& InOutOnAsyncRunCompletedDelegate, std::atomic<bool>& bIsBackgroundThreadRunning, const ENeuralNetworkSynchronousMode InSynchronousMode, const ENeuralDeviceType InDeviceType, const ENeuralDeviceType InInputDeviceType, const ENeuralDeviceType InOutputDeviceType);
+	void Run(FOnAsyncRunCompleted& InOutOnAsyncRunCompletedDelegate, std::atomic<bool>& bInIsBackgroundThreadRunning, const ENeuralNetworkSynchronousMode InSynchronousMode, const ENeuralDeviceType InDeviceType, const ENeuralDeviceType InInputDeviceType, const ENeuralDeviceType InOutputDeviceType);
 
 private:
 	/**
