Making sure tests pass after dropping UEOnly model tests.

_codereview Francisco.Vicente
#rnx

diff --git a/Engine/Plugins/Experimental/NeuralNetworkInference/Source/NeuralNetworkInference/Private/NeuralNetwork.cpp b/Engine/Plugins/Experimental/NeuralNetworkInference/Source/NeuralNetworkInference/Private/NeuralNetwork.cpp
--- a/Engine/Plugins/Experimental/NeuralNetworkInference/Source/NeuralNetworkInference/Private/NeuralNetwork.cpp
+++ b/Engine/Plugins/Experimental/NeuralNetworkInference/Source/NeuralNetworkInference/Private/NeuralNetwork.cpp
@@ -179,7 +179,7 @@ ENeuralBackEnd UNeuralNetwork::GetBackEndForCurrentPlatform() const
 	return BackEndForCurrentPlatform;
 }
 
-void UNeuralNetwork::SetBackEnd(const ENeuralBackEnd InBackEnd)
+bool UNeuralNetwork::SetBackEnd(const ENeuralBackEnd InBackEnd)
 {
 	BackEnd = InBackEnd;
 	const ENeuralBackEnd NewBackEndForCurrentPlatform = FPrivateNeuralNetwork::SetBackEndForCurrentPlatform(BackEnd);
@@ -193,6 +193,7 @@ void UNeuralNetwork::SetBackEnd(const ENeuralBackEnd InBackEnd)
 			Load();
 		}
 	}
+	return IsLoaded();
 }
 
 void UNeuralNetwork::ResetStats()
@@ -558,7 +559,7 @@ bool UNeuralNetwork::Load()
 	// Reset Stats
 	ResetStats();
 
-	return bIsLoaded;
+	return IsLoaded();
 }
 
 FNeuralTensor& UNeuralNetwork::GetInputTensorMutable(const int32 InTensorIndex)
@@ -585,7 +586,7 @@ bool UNeuralNetwork::Load(TArray<FNeuralTensor>& InTensors, const TArray<FNeural
 	}
 	// Load network
 	bIsLoaded = UNeuralNetwork::FImplBackEndUEOnly::Load(ImplBackEndUEOnly, TensorManager, InOperators);
-	return bIsLoaded;
+	return IsLoaded();
 }
 
 #if WITH_EDITOR
diff --git a/Engine/Plugins/Experimental/NeuralNetworkInference/Source/NeuralNetworkInference/Public/NeuralNetwork.h b/Engine/Plugins/Experimental/NeuralNetworkInference/Source/NeuralNetworkInference/Public/NeuralNetwork.h
--- a/Engine/Plugins/Experimental/NeuralNetworkInference/Source/NeuralNetworkInference/Public/NeuralNetwork.h
+++ b/Engine/Plugins/Experimental/NeuralNetworkInference/Source/NeuralNetworkInference/Public/NeuralNetwork.h
@@ -122,12 +122,12 @@ public:
 	 * GetBackEnd()/GetBackEndForCurrentPlatform():
 	 * - If BackEnd == Auto, GetBackEnd() will return Auto and GetBackEndForCurrentPlatform() will return the actual BackEnd being used for the current platform (UEAndORT or UEOnly).
 	 * - If BackEnd != Auto, GetBackEnd() and GetBackEndForCurrentPlatform() will both return the same value (UEAndORT or UEOnly).
-	 * SetBackEnd() will modify both BackEnd and BackEndForCurrentPlatform.
+	 * SetBackEnd() will modify both BackEnd and BackEndForCurrentPlatform and return IsLoaded().
 	 * @see ENeuralBackEnd for more details.
 	 */
 	ENeuralBackEnd GetBackEnd() const;
 	ENeuralBackEnd GetBackEndForCurrentPlatform() const;
-	void SetBackEnd(const ENeuralBackEnd InBackEnd);
+	bool SetBackEnd(const ENeuralBackEnd InBackEnd);
 	
 	/**
 	 * IsGPUConfigCompatible will always return true for ENeuralBackEnd::UEOnly. For ENeuralBackEnd::UEAndORT, IsGPUConfigCompatible() and IsGPUConfigCompatibleForUEAndORTBackEnd() will return the same value:
diff --git a/Engine/Plugins/Experimental/NeuralNetworkInference/Source/NeuralNetworkInferenceQA/Private/ModelUnitTester.cpp b/Engine/Plugins/Experimental/NeuralNetworkInference/Source/NeuralNetworkInferenceQA/Private/ModelUnitTester.cpp
--- a/Engine/Plugins/Experimental/NeuralNetworkInference/Source/NeuralNetworkInferenceQA/Private/ModelUnitTester.cpp
+++ b/Engine/Plugins/Experimental/NeuralNetworkInference/Source/NeuralNetworkInferenceQA/Private/ModelUnitTester.cpp
@@ -43,6 +43,7 @@ bool FModelUnitTester::GlobalTest(const FString& InProjectContentDir, const FStr
 	const TArray<int32> GPURepetitionsForUEOnlyBackEnd({ 0, 0, 0, 0 });
 #endif //WITH_EDITOR
 	// Run tests
+	UE_LOG(LogNeuralNetworkInferenceQA, Display, TEXT("ENeuralBackEnd::UEOnly test for FModelUnitTester::GlobalTest disabled due to not being compatible with ORT (linking error issues)."));
 	return ModelLoadAccuracyAndSpeedTests(InProjectContentDir, InModelZooRelativeDirectory, ModelNames, InputArrayValues, CPUGroundTruths, GPUGroundTruths, CPURepetitionsForUEAndORTBackEnd, GPURepetitionsForUEAndORTBackEnd, CPURepetitionsForUEOnlyBackEnd, GPURepetitionsForUEOnlyBackEnd);
 }
 
@@ -90,7 +91,7 @@ bool FModelUnitTester::ModelLoadAccuracyAndSpeedTests(const FString& InProjectCo
 		{
 			bDidGlobalTestPassed &= ModelAccuracyTest(Network, ENeuralBackEnd::UEAndORT, InInputArrayValues, CPUGroundTruths, GPUGroundTruths);
 		}
-		const bool bShouldRunUEOnlyBackEnd = (InCPURepetitionsForUEOnlyBackEnd[ModelIndex] * InCPURepetitionsForUEOnlyBackEnd[ModelIndex] > 0);
+		const bool bShouldRunUEOnlyBackEnd = false; // (InCPURepetitionsForUEOnlyBackEnd[ModelIndex] * InCPURepetitionsForUEOnlyBackEnd[ModelIndex] > 0);
 		if (bShouldRunUEOnlyBackEnd)
 		{
 			bDidGlobalTestPassed &= ModelAccuracyTest(Network, ENeuralBackEnd::UEOnly, InInputArrayValues, CPUGroundTruths, GPUGroundTruths);
@@ -144,8 +145,12 @@ bool FModelUnitTester::ModelLoadAccuracyAndSpeedTests(const FString& InProjectCo
 #endif //WITH_UE_AND_ORT_SUPPORT
 
 		// UEOnly
-		bDidGlobalTestPassed &= ModelSpeedTest(UAssetModelFilePath, ENeuralDeviceType::CPU, ENeuralBackEnd::UEOnly, InCPURepetitionsForUEOnlyBackEnd[ModelIndex]);
-		bDidGlobalTestPassed &= ModelSpeedTest(UAssetModelFilePath, ENeuralDeviceType::GPU, ENeuralBackEnd::UEOnly, InGPURepetitionsForUEOnlyBackEnd[ModelIndex]);
+		const bool bShouldRunUEOnlyBackEnd = false;
+		if (bShouldRunUEOnlyBackEnd)
+		{
+			bDidGlobalTestPassed &= ModelSpeedTest(UAssetModelFilePath, ENeuralDeviceType::CPU, ENeuralBackEnd::UEOnly, InCPURepetitionsForUEOnlyBackEnd[ModelIndex]);
+			bDidGlobalTestPassed &= ModelSpeedTest(UAssetModelFilePath, ENeuralDeviceType::GPU, ENeuralBackEnd::UEOnly, InGPURepetitionsForUEOnlyBackEnd[ModelIndex]);
+		}
 	}
 
 	return bDidGlobalTestPassed;
@@ -240,7 +245,11 @@ bool FModelUnitTester::ModelAccuracyTest(UNeuralNetwork* InOutNetwork, const ENe
 	const ENeuralDeviceType OriginalOutputDeviceType = InOutNetwork->GetOutputDeviceType();
 	const ENeuralBackEnd OriginalBackEnd = InOutNetwork->GetBackEnd();
 	// Set back end
-	InOutNetwork->SetBackEnd(InBackEnd);
+	if (!InOutNetwork->SetBackEnd(InBackEnd))
+	{
+		UE_LOG(LogNeuralNetworkInferenceQA, Warning, TEXT("Backend %s is disabled."), *GetBackEndString(InBackEnd));
+		return false;
+	}
 	const FString BackEndString = GetBackEndString(InBackEnd);
 	// Run each input with CPU/GPU and compare with each other and with the ground truth
 	// Multiple for loops to make sure that running on the CPU does not affect GPU results or vice versa
@@ -427,7 +436,11 @@ bool FModelUnitTester::ModelSpeedTest(const FString& InUAssetPath, const ENeural
 	const ENeuralDeviceType OriginalDeviceType = InOutNetwork->GetDeviceType();
 	const ENeuralBackEnd OriginalBackEnd = InOutNetwork->GetBackEnd();
 	// Set desired back end
-	InOutNetwork->SetBackEnd(InBackEnd);
+	if (!InOutNetwork->SetBackEnd(InBackEnd))
+	{
+		UE_LOG(LogNeuralNetworkInferenceQA, Warning, TEXT("Backend %s is disabled."), *GetBackEndString(InBackEnd));
+		return false;
+	}
 	// Needed variables
 	const int64 NetworkSize = InOutNetwork->GetInputTensor().Num();
 	TArray<float> InputArray;
diff --git a/Engine/Plugins/Experimental/NeuralNetworkInference/Source/Private/ModelProtoFileReader/Private/ModelProtoConverter.cpp b/Engine/Plugins/Experimental/NeuralNetworkInference/Source/Private/ModelProtoFileReader/Private/ModelProtoConverter.cpp
--- a/Engine/Plugins/Experimental/NeuralNetworkInference/Source/Private/ModelProtoFileReader/Private/ModelProtoConverter.cpp
+++ b/Engine/Plugins/Experimental/NeuralNetworkInference/Source/Private/ModelProtoFileReader/Private/ModelProtoConverter.cpp
@@ -538,8 +538,8 @@ bool FModelProtoConverter::ConvertFromONNXProto3(FModelProto& OutModelProto, std
 	#endif //WITH_PROTOBUF
 
 #else //WITH_MODEL_PROTO_CONVERTER_SUPPORT
-	UE_LOG(LogModelProtoFileReader, Warning,
-		TEXT("FModelProtoConverter::ConvertFromONNXProto3(): Due to a known linking error when used together with the UEAndORT back end during game mode, this function is disabled."));
+	UE_LOG(LogModelProtoFileReader, Warning, TEXT("FModelProtoConverter::ConvertFromONNXProto3(): Due to a known linking error when used together with the"
+		" UEAndORT back end during game mode, this function is disabled."));
 	return false;
 #endif //WITH_MODEL_PROTO_CONVERTER_SUPPORT
 }
