NNIEditor and NNIQAEditor combined into 1 (to simplify code reader's life).

_codereview Francisco.Vicente
#rnx

diff --git a/Engine/Plugins/Experimental/NeuralNetworkInference/NeuralNetworkInference.uplugin b/Engine/Plugins/Experimental/NeuralNetworkInference/NeuralNetworkInference.uplugin
--- a/Engine/Plugins/Experimental/NeuralNetworkInference/NeuralNetworkInference.uplugin
+++ b/Engine/Plugins/Experimental/NeuralNetworkInference/NeuralNetworkInference.uplugin
@@ -68,11 +68,6 @@
 			"Name": "NeuralNetworkInferenceEditor",
 			"Type": "Editor",
 			"LoadingPhase": "Default"
-		},
-		{
-			"Name": "NeuralNetworkInferenceQAEditor",
-			"Type": "Editor",
-			"LoadingPhase": "Default"
 		}
 	]
 }
\ No newline at end of file
diff --git a/Engine/Plugins/Experimental/NeuralNetworkInference/Source/NeuralNetworkInferenceEditor/NeuralNetworkInferenceEditor.Build.cs b/Engine/Plugins/Experimental/NeuralNetworkInference/Source/NeuralNetworkInferenceEditor/NeuralNetworkInferenceEditor.Build.cs
--- a/Engine/Plugins/Experimental/NeuralNetworkInference/Source/NeuralNetworkInferenceEditor/NeuralNetworkInferenceEditor.Build.cs
+++ b/Engine/Plugins/Experimental/NeuralNetworkInference/Source/NeuralNetworkInferenceEditor/NeuralNetworkInferenceEditor.Build.cs
@@ -28,6 +28,7 @@ public class NeuralNetworkInferenceEditor : ModuleRules
 			(
 			new string[] {
 				"NeuralNetworkInference",
+				"NeuralNetworkInferenceQA",
 				"UnrealEd"
 			}
 		);
diff --git a/Engine/Plugins/Experimental/NeuralNetworkInference/Source/NeuralNetworkInferenceEditor/Private/NeuralNetworkInferenceEditorModule.cpp b/Engine/Plugins/Experimental/NeuralNetworkInference/Source/NeuralNetworkInferenceEditor/Private/NeuralNetworkInferenceEditorModule.cpp
--- a/Engine/Plugins/Experimental/NeuralNetworkInference/Source/NeuralNetworkInferenceEditor/Private/NeuralNetworkInferenceEditorModule.cpp
+++ b/Engine/Plugins/Experimental/NeuralNetworkInference/Source/NeuralNetworkInferenceEditor/Private/NeuralNetworkInferenceEditorModule.cpp
@@ -1,9 +1,10 @@
 // Copyright Epic Games, Inc. All Rights Reserved.
 
 #include "NeuralNetworkInferenceEditorModule.h"
-#include "NeuralNetworkAssetTypeActions.h"
 #include "IAssetTypeActions.h"
 #include "Interfaces/IPluginManager.h"
+#include "NeuralNetworkAssetTypeActions.h"
+#include "QA/NeuralNetworkInferenceQAAssetTypeActions.h"
 
 
 
@@ -21,7 +22,7 @@ public:
 	virtual EAssetTypeCategories::Type GetMLAssetCategoryBit() const;
 
 private:
-	TSharedPtr<IAssetTypeActions> Action;
+	TArray<TSharedPtr<IAssetTypeActions>> Actions;
 	EAssetTypeCategories::Type MLAssetCategoryBit;
 };
 
@@ -33,10 +34,13 @@ private:
 // This code will execute after your module is loaded into memory; the exact timing is specified in the .uplugin file per-module
 void FNeuralNetworkInferenceEditorModule::StartupModule()
 {
-	// UNeuralNetwork - Register asset types
 	IAssetTools& AssetTools = FModuleManager::LoadModuleChecked<FAssetToolsModule>("AssetTools").Get();
-	Action = MakeShared<FNeuralNetworkAssetTypeActions>();
-	AssetTools.RegisterAssetTypeActions(Action.ToSharedRef());
+	// UNeuralNetwork - Register asset types
+	Actions.Emplace(MakeShared<FNeuralNetworkAssetTypeActions>());
+	AssetTools.RegisterAssetTypeActions(Actions.Last().ToSharedRef());
+	// NNI QA - Register asset types
+	Actions.Emplace(MakeShared<FNeuralNetworkAssetTypeActions>());
+	AssetTools.RegisterAssetTypeActions(Actions.Last().ToSharedRef());
 	// Register ML category so that ML assets can register to it
 	MLAssetCategoryBit = AssetTools.RegisterAdvancedAssetCategory(FName(TEXT("ML")), NSLOCTEXT("MLAssetCategory", "MLAssetCategory_ML", "Machine Learning"));
 }
@@ -50,9 +54,12 @@ void FNeuralNetworkInferenceEditorModule::ShutdownModule()
 	{
 		IAssetTools& AssetTools = ModuleInterface->Get();
 		// UNeuralNetwork - Unregister asset types
-		if (Action.IsValid())
+		for (TSharedPtr<IAssetTypeActions>& Action : Actions)
 		{
-			AssetTools.UnregisterAssetTypeActions(Action.ToSharedRef());
+			if (Action.IsValid())
+			{
+				AssetTools.UnregisterAssetTypeActions(Action.ToSharedRef());
+			}
 		}
 	}
 }
diff --git a/Engine/Plugins/Experimental/NeuralNetworkInference/Source/NeuralNetworkInferenceEditor/Private/QA/NeuralNetworkInferenceQAAssetFactory.cpp b/Engine/Plugins/Experimental/NeuralNetworkInference/Source/NeuralNetworkInferenceEditor/Private/QA/NeuralNetworkInferenceQAAssetFactory.cpp
new file mode 100644
--- /dev/null
+++ b/Engine/Plugins/Experimental/NeuralNetworkInference/Source/NeuralNetworkInferenceEditor/Private/QA/NeuralNetworkInferenceQAAssetFactory.cpp
@@ -0,0 +1,83 @@
+// Copyright Epic Games, Inc. All Rights Reserved.
+
+#include "NeuralNetworkInferenceQAAssetFactory.h"
+#include "NeuralNetworkInferenceQAAsset.h"
+#include "Misc/Paths.h"
+
+
+
+/* UNeuralNetworkInferenceQAAssetFactory structors
+ *****************************************************************************/
+
+UNeuralNetworkInferenceQAAssetFactory::UNeuralNetworkInferenceQAAssetFactory()
+{
+	SupportedClass = UNeuralNetworkInferenceQAAsset::StaticClass();
+
+	bCreateNew = true;
+	bEditorImport = true;
+	bEditAfterNew = true;
+	bText = false;
+}
+
+
+
+/* UNeuralNetworkInferenceQAAssetFactory public functions
+ *****************************************************************************/
+
+UObject* UNeuralNetworkInferenceQAAssetFactory::FactoryCreateNew(UClass* InClass, UObject* InParent, FName InName, EObjectFlags InFlags, UObject* InContext, FFeedbackContext* InWarn)
+{
+	// If created with right-click on Content Browser --> Neural Network
+	return NewObject<UNeuralNetworkInferenceQAAsset>(InParent, InName, InFlags);
+}
+
+UObject* UNeuralNetworkInferenceQAAssetFactory::FactoryCreateFile(UClass* InClass, UObject* InParent, FName InName, EObjectFlags InFlags,
+	const FString& InFilename, const TCHAR* Params, FFeedbackContext* Warn, bool& bOutOperationCanceled)
+{
+	// If created by dragging a new file into the UE Editor Content Browser
+	if (InFilename.Len() > 0)
+	{
+		ensureMsgf(false, TEXT("No files should trigger this."));
+	}
+
+	// If created with right-click on Content Browser --> NeuralNetworkInferenceQAAsset
+	return NewObject<UNeuralNetworkInferenceQAAsset>(InParent, InName, InFlags);
+}
+
+bool UNeuralNetworkInferenceQAAssetFactory::CanCreateNew() const
+{
+	// If true --> It will always call FactoryCreateNew(), not allowing me to use FactoryCreateFile().
+	// If false --> It will ignore the FactoryCreateFile (thus the txt file) when creating a new UNeuralNetworkInferenceQAAsset.
+
+	// Keep false while in production, set to true only when trying to create a new one
+	return false; // CurrentFilename.IsEmpty();
+}
+
+bool UNeuralNetworkInferenceQAAssetFactory::DoesSupportClass(UClass * Class)
+{
+	return (Class == UNeuralNetworkInferenceQAAsset::StaticClass());
+}
+
+UClass* UNeuralNetworkInferenceQAAssetFactory::ResolveSupportedClass()
+{
+	return UNeuralNetworkInferenceQAAsset::StaticClass();
+}
+
+bool UNeuralNetworkInferenceQAAssetFactory::FactoryCanImport(const FString& InFilename)
+{
+	return IsValidFile(InFilename);
+}
+
+bool UNeuralNetworkInferenceQAAssetFactory::CanImportBeCanceled() const
+{
+	return false;
+}
+
+
+
+/* UNeuralNetworkInferenceQAAssetFactory protected functions
+ *****************************************************************************/
+
+bool UNeuralNetworkInferenceQAAssetFactory::IsValidFile(const FString& InFilename) const
+{
+	return false;
+}
diff --git a/Engine/Plugins/Experimental/NeuralNetworkInference/Source/NeuralNetworkInferenceEditor/Private/QA/NeuralNetworkInferenceQAAssetFactory.h b/Engine/Plugins/Experimental/NeuralNetworkInference/Source/NeuralNetworkInferenceEditor/Private/QA/NeuralNetworkInferenceQAAssetFactory.h
new file mode 100644
--- /dev/null
+++ b/Engine/Plugins/Experimental/NeuralNetworkInference/Source/NeuralNetworkInferenceEditor/Private/QA/NeuralNetworkInferenceQAAssetFactory.h
@@ -0,0 +1,31 @@
+// Copyright Epic Games, Inc. All Rights Reserved.
+
+#pragma once
+
+#include "CoreMinimal.h"
+#include "UObject/ObjectMacros.h"
+#include "Factories/Factory.h"
+#include "NeuralNetworkInferenceQAAssetFactory.generated.h"
+
+UCLASS(hidecategories=Object)
+class UNeuralNetworkInferenceQAAssetFactory : public UFactory
+{
+	GENERATED_BODY()
+
+public:
+	UNeuralNetworkInferenceQAAssetFactory();
+	//~ Begin UFactory Interface
+	// UI creation (AssetTypeActions)
+	virtual UObject* FactoryCreateNew(UClass* InClass, UObject* InParent, FName InName, EObjectFlags InFlags, UObject* InContext, FFeedbackContext* InWarn) override;
+	virtual UObject* FactoryCreateFile(UClass* InClass, UObject* InParent, FName InName, EObjectFlags InFlags, const FString& InFilename,
+		const TCHAR* Params, FFeedbackContext* Warn, bool& bOutOperationCanceled) override;
+	virtual bool CanCreateNew() const override;
+	// Dragging and importing ONNX files into the Content Browser
+	virtual bool DoesSupportClass(UClass * Class) override;
+	virtual UClass* ResolveSupportedClass() override;
+	virtual bool FactoryCanImport(const FString& InFilename) override;
+	virtual bool CanImportBeCanceled() const override;
+
+protected:
+	virtual bool IsValidFile(const FString& InFilename) const final;
+};
diff --git a/Engine/Plugins/Experimental/NeuralNetworkInference/Source/NeuralNetworkInferenceEditor/Private/QA/NeuralNetworkInferenceQAAssetTypeActions.cpp b/Engine/Plugins/Experimental/NeuralNetworkInference/Source/NeuralNetworkInferenceEditor/Private/QA/NeuralNetworkInferenceQAAssetTypeActions.cpp
new file mode 100644
--- /dev/null
+++ b/Engine/Plugins/Experimental/NeuralNetworkInference/Source/NeuralNetworkInferenceEditor/Private/QA/NeuralNetworkInferenceQAAssetTypeActions.cpp
@@ -0,0 +1,39 @@
+// Copyright Epic Games, Inc. All Rights Reserved.
+
+#include "NeuralNetworkInferenceQAAssetTypeActions.h"
+#include "NeuralNetworkInferenceEditorModule.h"
+#include "NeuralNetworkInferenceQAAsset.h"
+
+
+FText FNeuralNetworkInferenceQAAssetTypeActions::GetName() const
+{
+	return NSLOCTEXT("AssetTypeActions", "AssetTypeActions_NeuralNetworkInferenceQAAsset", "Neural Network Inference QA Asset");
+}
+
+FColor FNeuralNetworkInferenceQAAssetTypeActions::GetTypeColor() const
+{
+	return FColor::Red;
+}
+
+UClass* FNeuralNetworkInferenceQAAssetTypeActions::GetSupportedClass() const
+{
+	return UNeuralNetworkInferenceQAAsset::StaticClass();
+}
+
+bool FNeuralNetworkInferenceQAAssetTypeActions::HasActions(const TArray<UObject*>& InObjects) const
+{
+	// Returns true if this class can supply actions for InObjects
+	return false;
+}
+
+uint32 FNeuralNetworkInferenceQAAssetTypeActions::GetCategories()
+{
+	const INeuralNetworkInferenceEditorModule& NeuralNetworkInferenceEditorModule = FModuleManager::GetModuleChecked<INeuralNetworkInferenceEditorModule>("NeuralNetworkInferenceEditor");
+	return NeuralNetworkInferenceEditorModule.GetMLAssetCategoryBit(); // Or EAssetTypeCategories::Misc
+}
+
+bool FNeuralNetworkInferenceQAAssetTypeActions::IsImportedAsset() const
+{
+	// Returns whether the asset was imported from an external source
+	return true;
+}
diff --git a/Engine/Plugins/Experimental/NeuralNetworkInference/Source/NeuralNetworkInferenceEditor/Private/QA/NeuralNetworkInferenceQAAssetTypeActions.h b/Engine/Plugins/Experimental/NeuralNetworkInference/Source/NeuralNetworkInferenceEditor/Private/QA/NeuralNetworkInferenceQAAssetTypeActions.h
new file mode 100644
--- /dev/null
+++ b/Engine/Plugins/Experimental/NeuralNetworkInference/Source/NeuralNetworkInferenceEditor/Private/QA/NeuralNetworkInferenceQAAssetTypeActions.h
@@ -0,0 +1,18 @@
+// Copyright Epic Games, Inc. All Rights Reserved.
+
+#pragma once
+
+#include "CoreMinimal.h"
+#include "AssetTypeActions_Base.h"
+
+class FNeuralNetworkInferenceQAAssetTypeActions : public FAssetTypeActions_Base
+{
+public:
+	// IAssetTypeActions Implementation
+	virtual FText GetName() const override;
+	virtual FColor GetTypeColor() const override;
+	virtual UClass* GetSupportedClass() const override;
+	virtual bool HasActions(const TArray<UObject*>& InObjects) const override;
+	virtual uint32 GetCategories() override;
+	virtual bool IsImportedAsset() const override;
+};
diff --git a/Engine/Plugins/Experimental/NeuralNetworkInference/Source/Private/NeuralNetworkInferenceQAEditor/NeuralNetworkInferenceQAEditor.Build.cs b/Engine/Plugins/Experimental/NeuralNetworkInference/Source/Private/NeuralNetworkInferenceQAEditor/NeuralNetworkInferenceQAEditor.Build.cs
deleted file mode 100644
diff --git a/Engine/Plugins/Experimental/NeuralNetworkInference/Source/Private/NeuralNetworkInferenceQAEditor/Private/Module.cpp b/Engine/Plugins/Experimental/NeuralNetworkInference/Source/Private/NeuralNetworkInferenceQAEditor/Private/Module.cpp
deleted file mode 100644
diff --git a/Engine/Plugins/Experimental/NeuralNetworkInference/Source/Private/NeuralNetworkInferenceQAEditor/Private/Module.h b/Engine/Plugins/Experimental/NeuralNetworkInference/Source/Private/NeuralNetworkInferenceQAEditor/Private/Module.h
deleted file mode 100644
diff --git a/Engine/Plugins/Experimental/NeuralNetworkInference/Source/Private/NeuralNetworkInferenceQAEditor/Private/NeuralNetworkInferenceQAAssetFactory.cpp b/Engine/Plugins/Experimental/NeuralNetworkInference/Source/Private/NeuralNetworkInferenceQAEditor/Private/NeuralNetworkInferenceQAAssetFactory.cpp
deleted file mode 100644
diff --git a/Engine/Plugins/Experimental/NeuralNetworkInference/Source/Private/NeuralNetworkInferenceQAEditor/Private/NeuralNetworkInferenceQAAssetFactory.h b/Engine/Plugins/Experimental/NeuralNetworkInference/Source/Private/NeuralNetworkInferenceQAEditor/Private/NeuralNetworkInferenceQAAssetFactory.h
deleted file mode 100644
diff --git a/Engine/Plugins/Experimental/NeuralNetworkInference/Source/Private/NeuralNetworkInferenceQAEditor/Private/NeuralNetworkInferenceQAAssetTypeActions.cpp b/Engine/Plugins/Experimental/NeuralNetworkInference/Source/Private/NeuralNetworkInferenceQAEditor/Private/NeuralNetworkInferenceQAAssetTypeActions.cpp
deleted file mode 100644
diff --git a/Engine/Plugins/Experimental/NeuralNetworkInference/Source/Private/NeuralNetworkInferenceQAEditor/Private/NeuralNetworkInferenceQAAssetTypeActions.h b/Engine/Plugins/Experimental/NeuralNetworkInference/Source/Private/NeuralNetworkInferenceQAEditor/Private/NeuralNetworkInferenceQAAssetTypeActions.h
deleted file mode 100644
