Moved "Private" into "Internal"

_codereview Francisco.Vicente
#rnx

diff --git a/Engine/Plugins/Experimental/NeuralNetworkInference/Source/Internal/ModelProto/Internal/ModelProto.h b/Engine/Plugins/Experimental/NeuralNetworkInference/Source/Internal/ModelProto/Internal/ModelProto.h
deleted file mode 100644
diff --git a/Engine/Plugins/Experimental/NeuralNetworkInference/Source/Internal/ModelProto/ModelProto.Build.cs b/Engine/Plugins/Experimental/NeuralNetworkInference/Source/Internal/ModelProto/ModelProto.Build.cs
deleted file mode 100644
diff --git a/Engine/Plugins/Experimental/NeuralNetworkInference/Source/Internal/ModelProto/Private/ModelProto.cpp b/Engine/Plugins/Experimental/NeuralNetworkInference/Source/Internal/ModelProto/Private/ModelProto.cpp
deleted file mode 100644
diff --git a/Engine/Plugins/Experimental/NeuralNetworkInference/Source/Internal/ModelProto/Private/ModelProtoModule.cpp b/Engine/Plugins/Experimental/NeuralNetworkInference/Source/Internal/ModelProto/Private/ModelProtoModule.cpp
deleted file mode 100644
diff --git a/Engine/Plugins/Experimental/NeuralNetworkInference/Source/Internal/ModelProto/Private/ModelProtoModule.h b/Engine/Plugins/Experimental/NeuralNetworkInference/Source/Internal/ModelProto/Private/ModelProtoModule.h
deleted file mode 100644
diff --git a/Engine/Plugins/Experimental/NeuralNetworkInference/Source/Internal/ModelProto/Private/ModelProtoStringParser.cpp b/Engine/Plugins/Experimental/NeuralNetworkInference/Source/Internal/ModelProto/Private/ModelProtoStringParser.cpp
deleted file mode 100644
diff --git a/Engine/Plugins/Experimental/NeuralNetworkInference/Source/Internal/ModelProto/Private/ModelProtoStringParser.h b/Engine/Plugins/Experimental/NeuralNetworkInference/Source/Internal/ModelProto/Private/ModelProtoStringParser.h
deleted file mode 100644
diff --git a/Engine/Plugins/Experimental/NeuralNetworkInference/Source/Internal/ModelProto/Private/ModelProtoUtils.h b/Engine/Plugins/Experimental/NeuralNetworkInference/Source/Internal/ModelProto/Private/ModelProtoUtils.h
deleted file mode 100644
diff --git a/Engine/Plugins/Experimental/NeuralNetworkInference/Source/Internal/NeuralNetworkInferenceShaders/Internal/NeuralNetworkInferenceShaders/BatchNormalizationCS.h b/Engine/Plugins/Experimental/NeuralNetworkInference/Source/Internal/NeuralNetworkInferenceShaders/Internal/NeuralNetworkInferenceShaders/BatchNormalizationCS.h
deleted file mode 100644
diff --git a/Engine/Plugins/Experimental/NeuralNetworkInference/Source/Internal/NeuralNetworkInferenceShaders/Internal/NeuralNetworkInferenceShaders/ConvBaseCS.h b/Engine/Plugins/Experimental/NeuralNetworkInference/Source/Internal/NeuralNetworkInferenceShaders/Internal/NeuralNetworkInferenceShaders/ConvBaseCS.h
deleted file mode 100644
diff --git a/Engine/Plugins/Experimental/NeuralNetworkInference/Source/Internal/NeuralNetworkInferenceShaders/Internal/NeuralNetworkInferenceShaders/ConvTransposeCS.h b/Engine/Plugins/Experimental/NeuralNetworkInference/Source/Internal/NeuralNetworkInferenceShaders/Internal/NeuralNetworkInferenceShaders/ConvTransposeCS.h
deleted file mode 100644
diff --git a/Engine/Plugins/Experimental/NeuralNetworkInference/Source/Internal/NeuralNetworkInferenceShaders/Internal/NeuralNetworkInferenceShaders/CopyCS.h b/Engine/Plugins/Experimental/NeuralNetworkInference/Source/Internal/NeuralNetworkInferenceShaders/Internal/NeuralNetworkInferenceShaders/CopyCS.h
deleted file mode 100644
diff --git a/Engine/Plugins/Experimental/NeuralNetworkInference/Source/Internal/NeuralNetworkInferenceShaders/Internal/NeuralNetworkInferenceShaders/ElementWiseCS.h b/Engine/Plugins/Experimental/NeuralNetworkInference/Source/Internal/NeuralNetworkInferenceShaders/Internal/NeuralNetworkInferenceShaders/ElementWiseCS.h
deleted file mode 100644
diff --git a/Engine/Plugins/Experimental/NeuralNetworkInference/Source/Internal/NeuralNetworkInferenceShaders/Internal/NeuralNetworkInferenceShaders/GemmCS.h b/Engine/Plugins/Experimental/NeuralNetworkInference/Source/Internal/NeuralNetworkInferenceShaders/Internal/NeuralNetworkInferenceShaders/GemmCS.h
deleted file mode 100644
diff --git a/Engine/Plugins/Experimental/NeuralNetworkInference/Source/Internal/NeuralNetworkInferenceShaders/Internal/NeuralNetworkInferenceShaders/MultidirectionalBroadcastCS.h b/Engine/Plugins/Experimental/NeuralNetworkInference/Source/Internal/NeuralNetworkInferenceShaders/Internal/NeuralNetworkInferenceShaders/MultidirectionalBroadcastCS.h
deleted file mode 100644
diff --git a/Engine/Plugins/Experimental/NeuralNetworkInference/Source/Internal/NeuralNetworkInferenceShaders/NeuralNetworkInferenceShaders.Build.cs b/Engine/Plugins/Experimental/NeuralNetworkInference/Source/Internal/NeuralNetworkInferenceShaders/NeuralNetworkInferenceShaders.Build.cs
deleted file mode 100644
diff --git a/Engine/Plugins/Experimental/NeuralNetworkInference/Source/Internal/NeuralNetworkInferenceShaders/Private/NeuralNetworkInferenceShaders/BatchNormalizationCS.cpp b/Engine/Plugins/Experimental/NeuralNetworkInference/Source/Internal/NeuralNetworkInferenceShaders/Private/NeuralNetworkInferenceShaders/BatchNormalizationCS.cpp
deleted file mode 100644
diff --git a/Engine/Plugins/Experimental/NeuralNetworkInference/Source/Internal/NeuralNetworkInferenceShaders/Private/NeuralNetworkInferenceShaders/ConvBaseCS.cpp b/Engine/Plugins/Experimental/NeuralNetworkInference/Source/Internal/NeuralNetworkInferenceShaders/Private/NeuralNetworkInferenceShaders/ConvBaseCS.cpp
deleted file mode 100644
diff --git a/Engine/Plugins/Experimental/NeuralNetworkInference/Source/Internal/NeuralNetworkInferenceShaders/Private/NeuralNetworkInferenceShaders/ConvTransposeCS.cpp b/Engine/Plugins/Experimental/NeuralNetworkInference/Source/Internal/NeuralNetworkInferenceShaders/Private/NeuralNetworkInferenceShaders/ConvTransposeCS.cpp
deleted file mode 100644
diff --git a/Engine/Plugins/Experimental/NeuralNetworkInference/Source/Internal/NeuralNetworkInferenceShaders/Private/NeuralNetworkInferenceShaders/CopyCS.cpp b/Engine/Plugins/Experimental/NeuralNetworkInference/Source/Internal/NeuralNetworkInferenceShaders/Private/NeuralNetworkInferenceShaders/CopyCS.cpp
deleted file mode 100644
diff --git a/Engine/Plugins/Experimental/NeuralNetworkInference/Source/Internal/NeuralNetworkInferenceShaders/Private/NeuralNetworkInferenceShaders/ElementWiseCS.cpp b/Engine/Plugins/Experimental/NeuralNetworkInference/Source/Internal/NeuralNetworkInferenceShaders/Private/NeuralNetworkInferenceShaders/ElementWiseCS.cpp
deleted file mode 100644
diff --git a/Engine/Plugins/Experimental/NeuralNetworkInference/Source/Internal/NeuralNetworkInferenceShaders/Private/NeuralNetworkInferenceShaders/GemmCS.cpp b/Engine/Plugins/Experimental/NeuralNetworkInference/Source/Internal/NeuralNetworkInferenceShaders/Private/NeuralNetworkInferenceShaders/GemmCS.cpp
deleted file mode 100644
diff --git a/Engine/Plugins/Experimental/NeuralNetworkInference/Source/Internal/NeuralNetworkInferenceShaders/Private/NeuralNetworkInferenceShaders/Module.cpp b/Engine/Plugins/Experimental/NeuralNetworkInference/Source/Internal/NeuralNetworkInferenceShaders/Private/NeuralNetworkInferenceShaders/Module.cpp
deleted file mode 100644
diff --git a/Engine/Plugins/Experimental/NeuralNetworkInference/Source/Internal/NeuralNetworkInferenceShaders/Private/NeuralNetworkInferenceShaders/Module.h b/Engine/Plugins/Experimental/NeuralNetworkInference/Source/Internal/NeuralNetworkInferenceShaders/Private/NeuralNetworkInferenceShaders/Module.h
deleted file mode 100644
diff --git a/Engine/Plugins/Experimental/NeuralNetworkInference/Source/Internal/NeuralNetworkInferenceShaders/Private/NeuralNetworkInferenceShaders/MultidirectionalBroadcastCS.cpp b/Engine/Plugins/Experimental/NeuralNetworkInference/Source/Internal/NeuralNetworkInferenceShaders/Private/NeuralNetworkInferenceShaders/MultidirectionalBroadcastCS.cpp
deleted file mode 100644
diff --git a/Engine/Plugins/Experimental/NeuralNetworkInference/Source/Internal/NeuralNetworkInferenceShaders/Private/NeuralNetworkInferenceShaders/Utils.h b/Engine/Plugins/Experimental/NeuralNetworkInference/Source/Internal/NeuralNetworkInferenceShaders/Private/NeuralNetworkInferenceShaders/Utils.h
deleted file mode 100644
diff --git a/Engine/Plugins/Experimental/NeuralNetworkInference/Source/Internal/ThirdPartyHelperAndDLLLoader/Internal/ORTExceptionHandler.h b/Engine/Plugins/Experimental/NeuralNetworkInference/Source/Internal/ThirdPartyHelperAndDLLLoader/Internal/ORTExceptionHandler.h
deleted file mode 100644
diff --git a/Engine/Plugins/Experimental/NeuralNetworkInference/Source/Internal/ThirdPartyHelperAndDLLLoader/Internal/RedirectCoutAndCerrToUeLog.h b/Engine/Plugins/Experimental/NeuralNetworkInference/Source/Internal/ThirdPartyHelperAndDLLLoader/Internal/RedirectCoutAndCerrToUeLog.h
deleted file mode 100644
diff --git a/Engine/Plugins/Experimental/NeuralNetworkInference/Source/Internal/ThirdPartyHelperAndDLLLoader/Internal/ThirdPartyWarningDisabler.h b/Engine/Plugins/Experimental/NeuralNetworkInference/Source/Internal/ThirdPartyHelperAndDLLLoader/Internal/ThirdPartyWarningDisabler.h
deleted file mode 100644
diff --git a/Engine/Plugins/Experimental/NeuralNetworkInference/Source/Internal/ThirdPartyHelperAndDLLLoader/Private/ORTExceptionHandler.cpp b/Engine/Plugins/Experimental/NeuralNetworkInference/Source/Internal/ThirdPartyHelperAndDLLLoader/Private/ORTExceptionHandler.cpp
deleted file mode 100644
diff --git a/Engine/Plugins/Experimental/NeuralNetworkInference/Source/Internal/ThirdPartyHelperAndDLLLoader/Private/RedirectCoutAndCerrToUeLog.cpp b/Engine/Plugins/Experimental/NeuralNetworkInference/Source/Internal/ThirdPartyHelperAndDLLLoader/Private/RedirectCoutAndCerrToUeLog.cpp
deleted file mode 100644
diff --git a/Engine/Plugins/Experimental/NeuralNetworkInference/Source/Internal/ThirdPartyHelperAndDLLLoader/Private/ThirdPartyHelperAndDLLLoaderModule.cpp b/Engine/Plugins/Experimental/NeuralNetworkInference/Source/Internal/ThirdPartyHelperAndDLLLoader/Private/ThirdPartyHelperAndDLLLoaderModule.cpp
deleted file mode 100644
diff --git a/Engine/Plugins/Experimental/NeuralNetworkInference/Source/Internal/ThirdPartyHelperAndDLLLoader/Private/ThirdPartyHelperAndDLLLoaderModule.h b/Engine/Plugins/Experimental/NeuralNetworkInference/Source/Internal/ThirdPartyHelperAndDLLLoader/Private/ThirdPartyHelperAndDLLLoaderModule.h
deleted file mode 100644
diff --git a/Engine/Plugins/Experimental/NeuralNetworkInference/Source/Internal/ThirdPartyHelperAndDLLLoader/Private/ThirdPartyHelperAndDLLLoaderUtils.h b/Engine/Plugins/Experimental/NeuralNetworkInference/Source/Internal/ThirdPartyHelperAndDLLLoader/Private/ThirdPartyHelperAndDLLLoaderUtils.h
deleted file mode 100644
diff --git a/Engine/Plugins/Experimental/NeuralNetworkInference/Source/Internal/ThirdPartyHelperAndDLLLoader/ThirdPartyHelperAndDLLLoader.Build.cs b/Engine/Plugins/Experimental/NeuralNetworkInference/Source/Internal/ThirdPartyHelperAndDLLLoader/ThirdPartyHelperAndDLLLoader.Build.cs
deleted file mode 100644
diff --git a/Engine/Plugins/Experimental/NeuralNetworkInference/Source/Private/ModelProto/Internal/ModelProto.h b/Engine/Plugins/Experimental/NeuralNetworkInference/Source/Private/ModelProto/Internal/ModelProto.h
new file mode 100644
--- /dev/null
+++ b/Engine/Plugins/Experimental/NeuralNetworkInference/Source/Private/ModelProto/Internal/ModelProto.h
@@ -0,0 +1,1001 @@
+// Copyright Epic Games, Inc. All Rights Reserved.
+
+#pragma once
+
+#include "CoreMinimal.h"
+#include "NeuralEnumClasses.h"
+#include "ModelProto.generated.h"
+
+
+
+/* Enum classes
+ *****************************************************************************/
+
+UENUM()
+enum class EAttributeProtoAttributeType
+{
+	UNDEFINED = 0,
+	FLOAT = 1,
+	INT = 2,
+	STRING = 3,
+	TENSOR = 4,
+	GRAPH = 5,
+	SPARSE_TENSOR = 11,
+	FLOATS = 6,
+	INTS = 7,
+	STRINGS = 8,
+	TENSORS = 9,
+	GRAPHS = 10,
+	SPARSE_TENSORS = 12
+};
+
+UENUM()
+enum class ETensorProtoDataType
+{
+	UNDEFINED = 0,
+	// Basic types.
+	FLOAT = 1,   // float
+	UINT8 = 2,   // uint8_t
+	INT8 = 3,    // int8_t
+	UINT16 = 4,  // uint16_t
+	INT16 = 5,   // int16_t
+	INT32 = 6,   // int32_t
+	INT64 = 7,   // int64_t
+	STRING = 8,  // string
+	BOOL = 9,    // bool
+
+	// IEEE754 half-precision floating-point format (16 bits wide).
+	// This format has 1 sign bit, 5 exponent bits, and 10 mantissa bits.
+	FLOAT16 = 10,
+
+	DOUBLE = 11,
+	UINT32 = 12,
+	UINT64 = 13,
+	COMPLEX64 = 14,     // complex with float32 real and imaginary components
+	COMPLEX128 = 15,    // complex with float64 real and imaginary components
+
+	// Non-IEEE floating-point format based on IEEE754 single-precision
+	// floating-point number truncated to 16 bits.
+	// This format has 1 sign bit, 8 exponent bits, and 7 mantissa bits.
+	BFLOAT16 = 16
+
+	// Future extensions go here.
+};
+
+UENUM()
+enum class ETensorProtoDataLocation
+{
+	DEFAULT = 0,
+	EXTERNAL = 1
+};
+
+
+
+/* Structs
+ *****************************************************************************/
+
+/**
+ * Level 7 - FTensorProtoSegment
+ * It follows the ONNX standard: [REDACTED_PROFILE_URL]/onnx/blob/master/onnx/onnx.proto3
+ */
+USTRUCT()
+struct MODELPROTO_API FTensorProtoSegment
+{
+	GENERATED_BODY()
+
+public:
+	UPROPERTY(VisibleAnywhere, Category = "Neural Network Inference")
+	int64 Begin;
+
+	UPROPERTY(VisibleAnywhere, Category = "Neural Network Inference")
+	int64 End;
+
+	FTensorProtoSegment();
+
+	bool LoadFromString(const FString& InProtoString, const int32 InLevel);
+
+	FString ToString(const FString& InLineStarted = TEXT("")) const;
+};
+
+/**
+ * Level 7 - FStringStringEntryProto
+ * It follows the ONNX standard: [REDACTED_PROFILE_URL]/onnx/blob/master/onnx/onnx.proto3
+ *
+ * StringStringEntryProto follows the pattern for cross-proto-version maps.
+ * See https://developers.google.com/protocol-buffers/docs/proto3#maps
+ */
+USTRUCT()
+struct MODELPROTO_API FStringStringEntryProto
+{
+	GENERATED_BODY()
+
+public:
+	UPROPERTY(VisibleAnywhere, Category = "Neural Network Inference")
+	FString Key;
+
+	UPROPERTY(VisibleAnywhere, Category = "Neural Network Inference")
+	FString Value;
+
+	bool LoadFromString(const FString& InProtoString, const int32 InLevel);
+
+	FString ToString(const FString& InLineStarted = TEXT("")) const;
+};
+
+/**
+ * Level 7 - FTensorShapeProtoDimension
+ * It follows the ONNX standard: [REDACTED_PROFILE_URL]/onnx/blob/master/onnx/onnx.proto3
+ */
+USTRUCT()
+struct MODELPROTO_API FTensorShapeProtoDimension
+{
+	GENERATED_BODY()
+
+public:
+	UPROPERTY(VisibleAnywhere, Category = "Neural Network Inference")
+	int64 DimValue;
+
+	UPROPERTY(VisibleAnywhere, Category = "Neural Network Inference")
+	FString DimParam;
+
+	UPROPERTY(VisibleAnywhere, Category = "Neural Network Inference")
+	FString Denotation;
+
+	FTensorShapeProtoDimension();
+
+	bool LoadFromString(const FString& InProtoString, const int32 InLevel);
+
+	FString ToString(const FString& InLineStarted = TEXT("")) const;
+};
+
+/**
+ * Level 6 - FTensorShapeProto
+ * It follows the ONNX standard: [REDACTED_PROFILE_URL]/onnx/blob/master/onnx/onnx.proto3
+ *
+ * Defines a tensor shape. A dimension can be either an integer value
+ * or a symbolic variable. A symbolic variable represents an unknown
+ * dimension.
+ */
+USTRUCT()
+struct MODELPROTO_API FTensorShapeProto
+{
+	GENERATED_BODY()
+
+public:
+	/**
+	 * A dimension can be either an integer value or a symbolic
+	 * variable. A symbolic variable represents an unknown dimension.
+	 */
+	UPROPERTY(VisibleAnywhere, Category = "Neural Network Inference")
+	TArray<FTensorShapeProtoDimension> Dim;
+
+	bool LoadFromString(const FString& InProtoString, const int32 InLevel);
+
+	FString ToString(const FString& InLineStarted = TEXT("")) const;
+};
+
+/**
+ * Level 6 - FTensorProto
+ * It follows the ONNX standard: [REDACTED_PROFILE_URL]/onnx/blob/master/onnx/onnx.proto3
+ *
+ * Tensors
+ * A serialized tensor value.
+ */
+USTRUCT()
+struct MODELPROTO_API FTensorProto
+{
+	GENERATED_BODY()
+
+public:
+	/**
+	 * The shape of the tensor.
+	 */
+	UPROPERTY(VisibleAnywhere, Category = "Neural Network Inference")
+	TArray<int64> Dimensions;
+
+	/**
+	 * The data type of the tensor.
+	 * This field MUST have a valid TensorProto.DataType value
+	 */
+	UPROPERTY(VisibleAnywhere, Category = "Neural Network Inference")
+	ETensorProtoDataType DataType;
+
+	/**
+	 * For very large tensors, we may want to store them in chunks, in which
+	 * case the following fields will specify the segment that is stored in
+	 * the current TensorProto.
+	 */
+	UPROPERTY(VisibleAnywhere, Category = "Neural Network Inference")
+	FTensorProtoSegment Segment;
+
+	UPROPERTY()
+	TArray<float> FloatData;
+	UPROPERTY()
+	TArray<int32> Int32Data;
+	UPROPERTY()
+	TArray<FString> StringData;
+	UPROPERTY()
+	TArray<int64> Int64Data;
+
+	/**
+	 * Optionally, a name for the tensor.
+	 */
+	UPROPERTY(VisibleAnywhere, Category = "Neural Network Inference")
+	FString Name;
+
+	/**
+	 * A human-readable documentation for this tensor. Markdown is allowed.
+	 */
+	UPROPERTY(VisibleAnywhere, Category = "Neural Network Inference")
+	FString DocString;
+
+
+	UPROPERTY()
+	TArray<uint8> RawData;
+	UPROPERTY(VisibleAnywhere, Category = "Neural Network Inference")
+	TArray<FStringStringEntryProto> ExternalData;
+
+	/**
+	 * If value not set, data is stored in raw_data (if set) otherwise in type-specified field.
+	 */
+	UPROPERTY(VisibleAnywhere, Category = "Neural Network Inference")
+	ETensorProtoDataLocation DataLocation;
+	UPROPERTY()
+	TArray<double> DoubleData;
+	UPROPERTY()
+	TArray<uint64> UInt64Data;
+
+	FTensorProto();
+
+	bool LoadFromString(const FString& InProtoString, const int32 InLevel);
+
+	FString ToString(const FString& InLineStarted = TEXT("")) const;
+
+	/**
+	 * It gets the data type from its ETensorProtoDataType analog. E.g.,
+	 * const ENeuralDataType NeuralDataType = FModelProto::GetDataTypeFromTensorProtoDataType(TensorProto.DataType);
+	 */
+	ENeuralDataType GetDataTypeFromTensorProtoDataType() const;
+
+private:
+	static ETensorProtoDataType StringToDataType(const FString& InString);
+
+	static ETensorProtoDataLocation StringToDataLocation(const FString& InString);
+};
+
+/**
+ * Level 5 - FTypeProtoSequence
+ * It follows the ONNX standard: [REDACTED_PROFILE_URL]/onnx/blob/master/onnx/onnx.proto3
+ */
+USTRUCT()
+struct MODELPROTO_API FTypeProtoSequence
+{
+	GENERATED_BODY()
+};
+
+/**
+ * Level 5 - FTypeProtoTensor
+ * It follows the ONNX standard: [REDACTED_PROFILE_URL]/onnx/blob/master/onnx/onnx.proto3
+ */
+USTRUCT()
+struct MODELPROTO_API FTypeProtoTensor
+{
+	GENERATED_BODY()
+
+public:
+	/**
+	 * This field MUST NOT have the value of UNDEFINED
+	 * This field MUST have a valid TensorProto.DataType value
+	 * This field MUST be present for this version of the IR.
+	 */
+	UPROPERTY(VisibleAnywhere, Category = "Neural Network Inference")
+	int32 ElemType;
+
+	UPROPERTY(VisibleAnywhere, Category = "Neural Network Inference")
+	FTensorShapeProto Shape;
+
+	FTypeProtoTensor();
+
+	bool LoadFromString(const FString& InProtoString, const int32 InLevel);
+
+	FString ToString(const FString& InLineStarted = TEXT("")) const;
+};
+
+/**
+ * Level 5 - FTypeProtoMap
+ * It follows the ONNX standard: [REDACTED_PROFILE_URL]/onnx/blob/master/onnx/onnx.proto3
+ */
+USTRUCT()
+struct MODELPROTO_API FTypeProtoMap
+{
+	GENERATED_BODY()
+
+public:
+	/**
+	 * This field MUST have a valid TensorProto.DataType value 
+	 * This field MUST be present for this version of the IR. 
+	 * This field MUST refer to an integral type ([U]INT{8|16|32|64}) or STRING
+	 */
+	UPROPERTY(VisibleAnywhere, Category = "Neural Network Inference")
+	int32 KeyType;
+
+	/**
+	 * This field MUST be present for this version of the IR.
+	 */
+	//UPROPERTY(VisibleAnywhere, Category = "Neural Network Inference")
+	// FTypeProto ValueType; // Not implemented yet because of Circular-dependency
+
+	FTypeProtoMap();
+};
+
+/**
+ * Level 5 - FSparseTensorProto
+ * It follows the ONNX standard: [REDACTED_PROFILE_URL]/onnx/blob/master/onnx/onnx.proto3
+ */
+USTRUCT()
+struct MODELPROTO_API FSparseTensorProto
+{
+	GENERATED_BODY()
+
+public:
+	/**
+	 * The sequence of non-default values are encoded as a tensor of shape [NNZ].
+	 * The default-value is zero for numeric tensors, and empty-string for string tensors.
+	 * values must have a non-empty name present which serves as a name for SparseTensorProto
+	 * when used in sparse_initializer list.
+	 */
+	UPROPERTY(VisibleAnywhere, Category = "Neural Network Inference")
+	FTensorProto Values;
+
+	/**
+	 * The indices of the non-default values, which may be stored in one of two formats.
+	 * (a) Indices can be a tensor of shape [NNZ, rank] with the [i,j]-th value
+	 * corresponding to the j-th index of the i-th value (in the values tensor).
+	 * (b) Indices can be a tensor of shape [NNZ], in which case the i-th value
+	 * must be the linearized-index of the i-th value (in the values tensor).
+	 * The linearized-index can be converted into an index tuple (k_1,...,k_rank)
+	 * using the shape provided below.
+	 * The indices must appear in ascending order without duplication.
+	 * In the first format, the ordering is lexicographic-ordering:
+	 * e.g., index-value [1,4] must appear before [2,1]
+	 */
+	UPROPERTY(VisibleAnywhere, Category = "Neural Network Inference")
+	FTensorProto Indices;
+
+	/**
+	 * The shape of the underlying dense-tensor: [dim_1, dim_2, ... dim_rank]
+	 */
+	UPROPERTY(VisibleAnywhere, Category = "Neural Network Inference")
+	TArray<int64> Dimensions;
+
+	bool LoadFromString(const FString& InProtoString, const int32 InLevel);
+
+	FString ToString(const FString& InLineStarted = TEXT("")) const;
+};
+
+/**
+ * Level 4 - FTypeProto
+ * It follows the ONNX standard: [REDACTED_PROFILE_URL]/onnx/blob/master/onnx/onnx.proto3
+ *
+ * Types
+ * The standard ONNX data types.
+ *
+ * NOTE:  DNN-only implementations of ONNX MAY elect to not support non-tensor values
+ * as input and output to graphs and nodes. These types are needed to naturally
+ * support classical ML operators. DNN operators SHOULD restrict their input
+ * and output types to tensors.
+ */
+USTRUCT()
+struct MODELPROTO_API FTypeProto
+{
+	GENERATED_BODY()
+
+public:
+	/**
+	 * The type of a tensor.
+	 */
+	UPROPERTY(VisibleAnywhere, Category = "Neural Network Inference")
+	FTypeProtoTensor TensorType;
+
+	/**
+	 * An optional denotation can be used to denote the whole
+	 * type with a standard semantic description as to what is
+	 * stored inside. Refer to [REDACTED_PROFILE_URL]/onnx/blob/master/docs/TypeDenotation.md#type-denotation-definition
+	 * for pre-defined type denotations.
+	 */
+	UPROPERTY(VisibleAnywhere, Category = "Neural Network Inference")
+	FString Denotation;
+
+	bool LoadFromString(const FString& InProtoString, const int32 InLevel);
+
+	FString ToString(const FString& InLineStarted = TEXT("")) const;
+};
+
+/**
+ * Level 4 - FAttributeProto
+ * It follows the ONNX standard: [REDACTED_PROFILE_URL]/onnx/blob/master/onnx/onnx.proto3
+ */
+USTRUCT()
+struct MODELPROTO_API FAttributeProto
+{
+	GENERATED_BODY()
+
+public:
+	/**
+	 * The name field MUST be present for this version of the IR.
+	 */
+	UPROPERTY(VisibleAnywhere, Category = "Neural Network Inference")
+	FString Name;
+
+	/**
+	 * A human-readable documentation for this attribute. Markdown is allowed.
+	 */
+	UPROPERTY(VisibleAnywhere, Category = "Neural Network Inference")
+	FString DocString;
+
+	/**
+	 * The type field MUST be present for this version of the IR.
+	 * For 0.0.1 versions of the IR, this field was not defined, and
+	 * implementations needed to use has_field heuristics to determine
+	 * which value field was in use.  For IR_VERSION 0.0.2 or later, this
+	 * field MUST be set and match the f|i|s|t|... field in use.  This
+	 * change was made to accommodate proto3 implementations.
+	 */
+	UPROPERTY(VisibleAnywhere, Category = "Neural Network Inference")
+	EAttributeProtoAttributeType Type;
+
+	/**
+	 * Exactly ONE of the following fields must be present for this version of the IR
+	 */
+	UPROPERTY(VisibleAnywhere, Category = "Neural Network Inference")
+	float F;
+	UPROPERTY(VisibleAnywhere, Category = "Neural Network Inference")
+	int64 I;
+	UPROPERTY(VisibleAnywhere, Category = "Neural Network Inference")
+	FString S;
+	UPROPERTY(VisibleAnywhere, Category = "Neural Network Inference")
+	FTensorProto T;
+// 	UPROPERTY(VisibleAnywhere, Category = "Neural Network Inference")
+// 	TSharedPtr<FGraphProto> G; /** Used TSharedPtr<> to break the circular dependency */
+	UPROPERTY(VisibleAnywhere, Category = "Neural Network Inference")
+	FSparseTensorProto SparseTensor;
+	UPROPERTY(VisibleAnywhere, Category = "Neural Network Inference")
+	TArray<float> Floats;
+	UPROPERTY(VisibleAnywhere, Category = "Neural Network Inference")
+	TArray<int64> Integers;
+	UPROPERTY(VisibleAnywhere, Category = "Neural Network Inference")
+	TArray<FString> Strings;
+	UPROPERTY(VisibleAnywhere, Category = "Neural Network Inference")
+	TArray<FTensorProto> Tensors;
+	UPROPERTY(VisibleAnywhere, Category = "Neural Network Inference")
+	TArray<FSparseTensorProto> SparseTensors;
+// 	UPROPERTY(VisibleAnywhere, Category = "Neural Network Inference")
+// 	TArray<TSharedPtr<FGraphProto>> Graphs; /** Used TSharedPtr<> to break the circular dependency */
+
+
+	FAttributeProto();
+
+	bool LoadFromString(const FString& InProtoString, const int32 InLevel);
+
+	FString ToString(const FString& InLineStarted = TEXT("")) const;
+
+private:
+	static EAttributeProtoAttributeType StringToAttributeType(const FString& InString);
+};
+
+/**
+ * Level 3 - FValueInfoProto
+ * It follows the ONNX standard: [REDACTED_PROFILE_URL]/onnx/blob/master/onnx/onnx.proto3
+ *
+ * Defines information on value, including the name, the type, and
+ * the shape of the value.
+ */
+USTRUCT()
+struct MODELPROTO_API FValueInfoProto
+{
+	GENERATED_BODY()
+
+public:
+	/**
+	 *This field MUST be present in this version of the IR.
+	 */
+	UPROPERTY(VisibleAnywhere, Category = "Neural Network Inference")
+	FString Name;
+
+	/**
+	 * This field MUST be present in this version of the IR for
+	 * inputs and outputs of the top-level graph.
+	 */
+	UPROPERTY(VisibleAnywhere, Category = "Neural Network Inference")
+	FTypeProto Type;
+
+	/**
+	 * A human-readable documentation for this value. Markdown is allowed.
+	 */
+	UPROPERTY(VisibleAnywhere, Category = "Neural Network Inference")
+	FString DocString;
+
+	bool LoadFromString(const FString& InProtoString, const int32 InLevel);
+
+	FString ToString(const FString& InLineStarted = TEXT("")) const;
+};
+
+/**
+ * Level 3 - FTensorAnnotation
+ * It follows the ONNX standard: [REDACTED_PROFILE_URL]/onnx/blob/master/onnx/onnx.proto3
+ */
+USTRUCT()
+struct MODELPROTO_API FTensorAnnotation
+{
+	GENERATED_BODY()
+
+public:
+	UPROPERTY(VisibleAnywhere, Category = "Neural Network Inference")
+	FString TensorName;
+
+	/**
+	* <key, value> pairs to annotate tensor specified by <tensor_name> above.
+	* The keys used in the mapping below must be pre-defined in ONNX spec.
+	* For example, for 8-bit linear quantization case, 'SCALE_TENSOR',
+	* 'ZERO_POINT_TENSOR' will be pre-defined as quantization parameter keys.
+	*/
+	UPROPERTY(VisibleAnywhere, Category = "Neural Network Inference")
+	TArray<FStringStringEntryProto> QuantParameterTensorNames;
+
+	bool LoadFromString(const FString& InProtoString, const int32 InLevel);
+
+	FString ToString(const FString& InLineStarted = TEXT("")) const;
+};
+
+/**
+ * Level 3 - FNodeProto
+ * It follows the ONNX standard: [REDACTED_PROFILE_URL]/onnx/blob/master/onnx/onnx.proto3
+ *
+ * Nodes
+ *
+ * Computation graphs are made up of a DAG of nodes, which represent what is
+ * commonly called a "layer" or "pipeline stage" in machine learning frameworks.
+ *
+ * For example, it can be a node of type "Conv" that takes in an image, a filter
+ * tensor and a bias tensor, and produces the convolved output.
+ */
+USTRUCT()
+struct MODELPROTO_API FNodeProto
+{
+	GENERATED_BODY()
+
+public:
+	UPROPERTY(VisibleAnywhere, Category = "Neural Network Inference")
+	TArray<FString> Input;
+	UPROPERTY(VisibleAnywhere, Category = "Neural Network Inference")
+	TArray<FString> Output;
+
+	/**
+	 * An optional identifier for this node in a graph.
+	 * This field MAY be absent in this version of the IR.
+	 */
+	UPROPERTY(VisibleAnywhere, Category = "Neural Network Inference")
+	FString Name;
+
+	/**
+	 * The symbolic identifier of the Operator to execute.
+	 */
+	UPROPERTY(VisibleAnywhere, Category = "Neural Network Inference")
+	FString OperatorType;
+
+	/**
+	 * The domain of the OperatorSet that specifies the operator named by op_type.
+	 */
+	UPROPERTY(VisibleAnywhere, Category = "Neural Network Inference")
+	FString Domain;
+
+	/**
+	 * Additional named attributes.
+	 */
+	UPROPERTY(VisibleAnywhere, Category = "Neural Network Inference")
+	TArray<FAttributeProto> Attribute;
+
+	/**
+	 * A human-readable documentation for this node. Markdown is allowed.
+	 */
+	UPROPERTY(VisibleAnywhere, Category = "Neural Network Inference")
+	FString DocString;
+
+	bool LoadFromString(const FString& InProtoString, const int32 InLevel);
+
+	FString ToString(const FString& InLineStarted = TEXT("")) const;
+};
+
+/**
+ * Level 2 - FGraphProto
+ * It follows the ONNX standard: [REDACTED_PROFILE_URL]/onnx/blob/master/onnx/onnx.proto3
+ *
+ * Graphs
+ *
+ * A graph defines the computational logic of a model and is comprised of a parameterized
+ * list of nodes that form a directed acyclic graph based on their inputs and outputs.
+ * This is the equivalent of the "network" or "graph" in many deep learning / neural network
+ * frameworks.
+ */
+USTRUCT()
+struct MODELPROTO_API FGraphProto
+{
+	GENERATED_BODY()
+
+public:
+	UPROPERTY(VisibleAnywhere, Category = "Neural Network Inference")
+	bool bIsLoaded;
+
+	/**
+	 * The nodes in the graph, sorted topologically.
+	 */
+	UPROPERTY(VisibleAnywhere, Category = "Neural Network Inference")
+	TArray<FNodeProto> Node;
+
+	/**
+	 * The name of the graph.
+	 */
+	UPROPERTY(VisibleAnywhere, Category = "Neural Network Inference")
+	FString Name;
+
+	/**
+	 * A list of named tensor values, used to specify constant inputs of the graph.
+	 * Each TensorProto entry must have a distinct name (within the list) that
+	 * MAY also appear in the input list.
+	 */
+	UPROPERTY(VisibleAnywhere, Category = "Neural Network Inference")
+	TArray<FTensorProto> Initializer;
+
+	/**
+	 * Initializers (see above) stored in sparse format.
+	 */
+	UPROPERTY(VisibleAnywhere, Category = "Neural Network Inference")
+	TArray<FSparseTensorProto> SparseInitializer;
+
+	/**
+	 * A human-readable documentation for this graph. Markdown is allowed.
+	 */
+	UPROPERTY(VisibleAnywhere, Category = "Neural Network Inference")
+	FString DocString;
+
+	/**
+	 * The inputs and outputs of the graph.
+	 */
+	UPROPERTY(VisibleAnywhere, Category = "Neural Network Inference")
+	TArray<FValueInfoProto> Input;
+	UPROPERTY(VisibleAnywhere, Category = "Neural Network Inference")
+	TArray<FValueInfoProto> Output;
+
+	/**
+	 * Information for the values in the graph. The ValueInfoProto.name's
+	 * must be distinct. It is optional for a value to appear in value_info list.
+	 */
+	UPROPERTY(VisibleAnywhere, Category = "Neural Network Inference")
+	TArray<FValueInfoProto> ValueInfo;
+
+	/**
+	 * This field carries information to indicate the mapping among a tensor and its
+	 * quantization parameter tensors. For example:
+	 * For tensor 'a', it may have {'SCALE_TENSOR', 'a_scale'} and {'ZERO_POINT_TENSOR', 'a_zero_point'} annotated,
+	 * which means, tensor 'a_scale' and tensor 'a_zero_point' are scale and zero point of tensor 'a' in the model.
+	 */
+	TArray<FTensorAnnotation> QuantizationAnnotation;
+
+	FGraphProto();
+
+	FORCEINLINE bool IsLoaded() { return bIsLoaded; }
+
+	bool LoadFromString(const FString& InProtoString, const int32 InLevel);
+
+	FString ToString(const FString& InLineStarted = TEXT("")) const;
+};
+
+/**
+ * Level 1 - FTrainingInfoProto
+ * It follows the ONNX standard: [REDACTED_PROFILE_URL]/onnx/blob/master/onnx/onnx.proto3
+ */
+USTRUCT()
+struct MODELPROTO_API FTrainingInfoProto
+{
+	GENERATED_BODY()
+
+public:
+	UPROPERTY(VisibleAnywhere, Category = "Neural Network Inference")
+	bool bIsLoaded;
+
+	/**
+	 * This field describes a graph to compute the initial tensors
+	 * upon starting the training process. Initialization graph has no input
+	 * and can have multiple outputs. Usually, trainable tensors in neural 
+	 * networks are randomly initialized. To achieve that, for each tensor,
+     * the user can put a random number operator such as RandomNormal or
+	 * RandomUniform in TrainingInfoProto.initialization.node and assign its
+	 * random output to the specific tensor using "initialization_binding".
+	 * This graph can also set the initializers in "algorithm" in the same
+	 * TrainingInfoProto; a use case is resetting the number of training
+	 * iteration to zero.
+	 * 
+	 * By default, this field is an empty graph and its evaluation does not 
+	 * produce any output. Thus, no initializer would be changed by default.
+	 */
+	UPROPERTY(VisibleAnywhere, Category = "Neural Network Inference");
+	FGraphProto Initialization;
+
+	/**
+	 * This field represents a training algorithm step. Given required inputs,
+	 * it computes outputs to update initializers in its own or inference graph's 
+	 * initializer lists. In general, this field contains loss node, gradient node,
+	 * optimizer node, increment of iteration count.
+	 * 
+	 * An execution of the training algorithm step is performed by executing the
+	 * graph obtained by combining the inference graph (namely "ModelProto.graph")
+	 * and the "algorithm" graph. That is, the actual the actual
+	 * input/initializer/output/node/value_info/sparse_initializer list of
+	 * the training graph is the concatenation of
+	 * "ModelProto.graph.input/initializer/output/node/value_info/sparse_initializer"
+	 * and "algorithm.input/initializer/output/node/value_info/sparse_initializer"
+	 * in that order. This combined graph must satisfy the normal ONNX conditions.
+	 * Now, let's provide a visualization of graph combination for clarity.
+	 * Let the inference graph (i.e., "ModelProto.graph") be
+	 *    tensor_a, tensor_b -> MatMul -> tensor_c -> Sigmoid -> tensor_d
+	 * and the "algorithm" graph be
+	 *    tensor_d -> Add -> tensor_e
+	 * The combination process results 
+	 *    tensor_a, tensor_b -> MatMul -> tensor_c -> Sigmoid -> tensor_d -> Add -> tensor_e
+	 * 
+	 * Notice that an input of a node in the "algorithm" graph may reference the
+	 * output of a node in the inference graph (but not the other way round). Also, inference
+	 * node cannot reference inputs of "algorithm". With these restrictions, inference graph
+	 * can always be run independently without training information.
+	 * 
+	 * By default, this field is an empty graph and its evaluation does not
+	 * produce any output. Evaluating the default training step never
+	 * update any initializers.
+	 */
+	UPROPERTY(VisibleAnywhere, Category = "Neural Network Inference");
+	FGraphProto Algorithm;
+
+	/**
+	 * This field specifies the bindings from the outputs of "initialization" to
+	 * some initializers in "ModelProto.graph.initializer" and
+	 * some initializers in "ModelProto.graph.initializer" and
+	 * the "algorithm.initializer" in the same TrainingInfoProto.
+	 * See "update_binding" below for details.
+	 * 
+	 * By default, this field is empty and no initializer would be changed
+	 * by the execution of "initialization".
+	*/
+	UPROPERTY(VisibleAnywhere, Category = "Neural Network Inference")
+	TArray<FStringStringEntryProto> InitializationBinding;
+
+	/**
+	 * Gradient-based training is usually an iterative procedure. In one gradient
+	 * descent iteration, we apply
+	 * 
+	 * x = x - r * g
+	 * 
+	 * where "x" is the optimized tensor, "r" stands for learning rate, and "g" is
+	 * gradient of "x" with respect to a chosen loss. To avoid adding assignments
+	 * into the training graph, we split the update equation into
+	 * 
+	 * y = x - r * g
+	 * x = y
+	 * 
+	 * The user needs to save "y = x - r * g" into TrainingInfoProto.algorithm. To
+	 * tell that "y" should be assigned to "x", the field "update_binding" may
+	 * contain a key-value pair of strings, "x" (key of StringStringEntryProto)
+	 * and "y" (value of StringStringEntryProto).
+	 * For a neural network with multiple trainable (mutable) tensors, there can
+	 * be multiple key-value pairs in "update_binding".
+	 * 
+	 * The initializers appears as keys in "update_binding" are considered
+	 * mutable variables. This implies some behaviors
+	 * as described below.
+	 * 
+	 * 1. We have only unique keys in all "update_binding"s so that two
+	 *	  variables may not have the same name. This ensures that one
+	 *	  variable is assigned up to once.
+	 * 2. The keys must appear in names of "ModelProto.graph.initializer" or
+	 *    "TrainingInfoProto.algorithm.initializer".
+	 * 3. The values must be output names of "algorithm" or "ModelProto.graph.output".
+	 * 4. Mutable variables are initialized to the value specified by the
+	 *    corresponding initializer, and then potentially updated by
+	 *    "initializer_binding"s and "update_binding"s in "TrainingInfoProto"s.
+	 * 
+	 * This field usually contains names of trainable tensors
+	 * (in ModelProto.graph), optimizer states such as momentums in advanced
+	 * stochastic gradient methods (in TrainingInfoProto.graph),
+	 * and number of training iterations (in TrainingInfoProto.graph).
+	 * 
+	 * By default, this field is empty and no initializer would be changed
+	 * by the execution of "algorithm".
+	*/
+	UPROPERTY(VisibleAnywhere, Category = "Neural Network Inference")
+	TArray<FStringStringEntryProto> UpdateBinding;
+
+	FTrainingInfoProto();
+
+	FORCEINLINE bool IsLoaded() { return bIsLoaded; }
+
+	bool LoadFromString(const FString& InProtoString, const int32 InLevel);
+
+	FString ToString(const FString& InLineStarted = TEXT("")) const;
+};
+
+/**
+ * Level 1 - FOperatorSetIdProto
+ * It follows the ONNX standard: [REDACTED_PROFILE_URL]/onnx/blob/master/onnx/onnx.proto3
+ *
+ * Operator Sets
+ * OperatorSets are uniquely identified by a (domain, opset_version) pair.
+ */
+USTRUCT()
+struct MODELPROTO_API FOperatorSetIdProto
+{
+	GENERATED_BODY()
+
+public:
+	/**
+	 * The domain of the operator set being identified.
+	 * The empty string ("") or absence of this field implies the operator
+	 * set that is defined as part of the ONNX specification.
+	 * This field MUST be present in this version of the IR when referring to any other operator set.
+	 */
+	UPROPERTY(VisibleAnywhere, Category = "Neural Network Inference")
+	FString Domain;
+
+	/**
+	 * The version of the operator set being identified.
+	 * This field MUST be present in this version of the IR.
+	 */
+	UPROPERTY(VisibleAnywhere, Category = "Neural Network Inference")
+	int64 Version;
+
+	FOperatorSetIdProto();
+
+	bool LoadFromString(const FString& InProtoString, const int32 InLevel);
+
+	FString ToString(const FString& InLineStarted = TEXT("")) const;
+};
+
+/**
+ * Level 0 - FModelProto
+ * It follows the ONNX standard: [REDACTED_PROFILE_URL]/onnx/blob/master/onnx/onnx.proto3
+ *
+ * Models
+ *
+ * ModelProto is a top-level file/container format for bundling a ML model and
+ * associating its computation graph with metadata.
+ *
+ * The semantics of the model are described by the associated GraphProto's.
+ */
+USTRUCT()
+struct MODELPROTO_API FModelProto
+{
+	GENERATED_BODY()
+
+public:
+	UPROPERTY(VisibleAnywhere, Category = "Neural Network Inference")
+	bool bIsLoaded;
+
+	/**
+	 * The version of the IR this model targets. See Version enum above.
+	 * This field MUST be present.
+	 */
+	UPROPERTY(VisibleAnywhere, Category = "Neural Network Inference");
+	int64 IRVersion;
+
+	/**
+	 * The OperatorSets this model relies on.
+	 * All ModelProtos MUST have at least one entry that
+	 * specifies which version of the ONNX OperatorSet is
+	 * being imported.
+	 *
+	 * All nodes in the ModelProto's graph will bind against the operator
+	 * with the same-domain/same-op_type operator with the HIGHEST version
+	 * in the referenced operator sets.
+	 */
+	UPROPERTY(VisibleAnywhere, Category = "Neural Network Inference")
+	TArray<FOperatorSetIdProto> OpsetImport;
+
+	/**
+	 * The name of the framework or tool used to generate this model.
+	 * This field SHOULD be present to indicate which implementation/tool/framework
+	 * emitted the model.
+	 */
+	UPROPERTY(VisibleAnywhere, Category = "Neural Network Inference");
+	FString ProducerName;
+
+	/**
+	 * The version of the framework or tool used to generate this model.
+	 * This field SHOULD be present to indicate which implementation/tool/framework
+	 * emitted the model.
+	 */
+	UPROPERTY(VisibleAnywhere, Category = "Neural Network Inference");
+	FString ProducerVersion;
+
+	/**
+	 * Domain name of the model.
+     * We use reverse domain names as name space indicators. For example:
+     * `com.facebook.fair` or `com.microsoft.cognitiveservices`
+     *
+     * Together with `model_version` and GraphProto.name, this forms the unique identity of
+     * the graph.
+	 */
+	UPROPERTY(VisibleAnywhere, Category = "Neural Network Inference");
+	FString Domain;
+
+	/**
+	 * The version of the graph encoded. See Version enum below.
+	 */
+	UPROPERTY(VisibleAnywhere, Category = "Neural Network Inference");
+	int64 ModelVersion;
+
+	/**
+	 * A human-readable documentation for this model. Markdown is allowed.
+	 */
+	UPROPERTY(VisibleAnywhere, Category = "Neural Network Inference");
+	FString DocString;
+
+	/**
+	 * The parameterized graph that is evaluated to execute the model.
+	 */
+	UPROPERTY(VisibleAnywhere, Category = "Neural Network Inference");
+	FGraphProto Graph;
+
+	/**
+	 * Named metadata values; keys should be distinct.
+	 */
+	UPROPERTY(VisibleAnywhere, Category = "Neural Network Inference")
+	TArray<FStringStringEntryProto> MetadataProps;
+
+	/**
+	* Training-specific information. Sequentially executing all stored
+	 * `TrainingInfoProto.algorithm`s and assigning their outputs following
+	 * the corresponding `TrainingInfoProto.update_binding`s is one training
+	 * iteration. Similarly, to initialize the model (as if training hasn't
+	 * happened), the user should sequentially execute all stored
+	 * `TrainingInfoProto.initialization`s and assigns their outputs using
+	 * `TrainingInfoProto.initialization_binding`s.
+	 * If this field is empty, the training behavior of the model is undefined.
+	 */
+	UPROPERTY(VisibleAnywhere, Category = "Neural Network Inference")
+	TArray<FTrainingInfoProto> TrainingInfo;
+
+	FModelProto();
+
+	FORCEINLINE bool IsLoaded() { return bIsLoaded; }
+
+	bool LoadFromString(const FString& InProtoString, const int32 InLevel = 0);
+
+	FString ToString(const FString& InLineStarted = TEXT("")) const;
+
+	const FGraphProto& GetGraph() const;
+
+	/**
+	 * Given the TArray<T>, it will find and return the reference to the Element named InName.
+	 */
+	template <typename T>
+	static const T* FindElementInArray(const FString& InName, const TArray<T>& InArray, const bool bInMustValueBeFound);
+};
+
+
+
+/* ModelProto template functions
+ *****************************************************************************/
+
+template <typename T>
+const T* FModelProto::FindElementInArray(const FString& InName, const TArray<T>& InArray, const bool bInMustValueBeFound)
+{
+	// Iterating over each element name in InArray
+	for (const T& Element : InArray)
+	{
+		if (Element.Name == InName)
+		{
+			return &Element;
+		}
+	}
+
+	// Trigger assertion if tensor name does not exist and return
+	checkf(!bInMustValueBeFound, TEXT("Element with the name %s could not be found in InArray."), *InName);
+	return nullptr;
+}
diff --git a/Engine/Plugins/Experimental/NeuralNetworkInference/Source/Private/ModelProto/ModelProto.Build.cs b/Engine/Plugins/Experimental/NeuralNetworkInference/Source/Private/ModelProto/ModelProto.Build.cs
new file mode 100644
--- /dev/null
+++ b/Engine/Plugins/Experimental/NeuralNetworkInference/Source/Private/ModelProto/ModelProto.Build.cs
@@ -0,0 +1,27 @@
+// Copyright Epic Games, Inc. All Rights Reserved.
+
+using UnrealBuildTool;
+
+public class ModelProto : ModuleRules
+{
+	public ModelProto( ReadOnlyTargetRules Target ) : base( Target )
+	{
+		PCHUsage = ModuleRules.PCHUsageMode.UseExplicitOrSharedPCHs;
+
+		PublicIncludePaths.AddRange(
+			new string[] {
+				System.IO.Path.Combine(ModuleDirectory, "..")
+			}
+		);
+
+		PublicDependencyModuleNames.AddRange
+			(
+			new string[] {
+				"Core",
+				"CoreUObject",
+				"Engine",
+				"NeuralNetworkInferenceCore"
+			}
+		);
+	}
+}
diff --git a/Engine/Plugins/Experimental/NeuralNetworkInference/Source/Private/ModelProto/Private/ModelProto.cpp b/Engine/Plugins/Experimental/NeuralNetworkInference/Source/Private/ModelProto/Private/ModelProto.cpp
new file mode 100644
--- /dev/null
+++ b/Engine/Plugins/Experimental/NeuralNetworkInference/Source/Private/ModelProto/Private/ModelProto.cpp
@@ -0,0 +1,1066 @@
+// Copyright Epic Games, Inc. All Rights Reserved.
+
+#include "ModelProto.h"
+#include "ModelProtoUtils.h"
+#include "ModelProtoStringParser.h"
+
+
+
+/* FTensorProtoSegment structors
+ *****************************************************************************/
+
+FTensorProtoSegment::FTensorProtoSegment()
+	: Begin(-1)
+	, End(-1)
+{
+}
+
+/* FTensorProtoSegment public functions
+ *****************************************************************************/
+
+bool FTensorProtoSegment::LoadFromString(const FString& InProtoString, const int32 InLevel)
+{
+	// Sanity check
+	if (InProtoString.Len() < 1)
+	{
+		return false;
+	}
+
+	// FString text into [key, value] maps
+	const TMap<FString, TArray<FString>> ProtoMap = FModelProtoStringParser::ModelProtoStringToMap(InProtoString, InLevel);
+
+	// Fill class variables
+	Begin = FModelProtoStringParser::GetModelProtoInt64(ProtoMap, TEXT("begin"));
+	End = FModelProtoStringParser::GetModelProtoInt64(ProtoMap, TEXT("end"));
+
+	const bool bIsLoaded = Begin <= End;
+	if (!bIsLoaded)
+	{
+		UE_LOG(LogModelProto, Warning, TEXT("FTensorProtoSegment::LoadFromString() failed."));
+	}
+	return bIsLoaded;
+}
+
+FString FTensorProtoSegment::ToString(const FString&) const
+{
+	return FString::Format(TEXT("FTensorProtoSegment - [Begin, End] = [{0},{1}]:\n"), { FString::FromInt(Begin), FString::FromInt(End) });
+}
+
+
+
+/* FStringStringEntryProto public functions
+ *****************************************************************************/
+
+bool FStringStringEntryProto::LoadFromString(const FString& InProtoString, const int32 InLevel)
+{
+	// Sanity check
+	if (InProtoString.Len() < 1)
+	{
+		return false;
+	}
+
+	// FString text into [key, value] maps
+	const TMap<FString, TArray<FString>> ProtoMap = FModelProtoStringParser::ModelProtoStringToMap(InProtoString, InLevel);
+
+	// Fill class variables
+	Key = FModelProtoStringParser::GetModelProtoStringOrEmpty(ProtoMap, TEXT("key"));
+	Value = FModelProtoStringParser::GetModelProtoStringOrEmpty(ProtoMap, TEXT("value"));
+
+	const bool bIsLoaded = !Key.IsEmpty() && !Value.IsEmpty();
+	if (!bIsLoaded)
+	{
+		UE_LOG(LogModelProto, Warning, TEXT("FStringStringEntryProto::LoadFromString() failed."));
+	}
+	return bIsLoaded;
+}
+
+FString FStringStringEntryProto::ToString(const FString&) const
+{
+	return FString::Format(TEXT("FStringStringEntryProto: [Key, Value] = [{0},{1}]\n"), { Key, Value });
+}
+
+
+
+/* FTensorShapeProtoDimension structors
+ *****************************************************************************/
+
+FTensorShapeProtoDimension::FTensorShapeProtoDimension()
+	: DimValue(-1)
+{
+}
+
+/* FTensorShapeProtoDimension public functions
+ *****************************************************************************/
+
+bool FTensorShapeProtoDimension::LoadFromString(const FString& InProtoString, const int32 InLevel)
+{
+	// Sanity check
+	if (InProtoString.Len() < 1)
+	{
+		return false;
+	}
+
+	// FString text into [key, value] maps
+	const TMap<FString, TArray<FString>> ProtoMap = FModelProtoStringParser::ModelProtoStringToMap(InProtoString, InLevel);
+
+	// Fill class variables
+	DimValue = FModelProtoStringParser::GetModelProtoInt64(ProtoMap, TEXT("dim_value"));
+	DimParam = FModelProtoStringParser::GetModelProtoStringOrEmpty(ProtoMap, TEXT("dim_param"));
+
+	const bool bIsLoaded = (DimValue > -1);
+	if (!bIsLoaded)
+	{
+		UE_LOG(LogModelProto, Warning, TEXT("FTensorShapeProtoDimension::LoadFromString(): failed."));
+	}
+	return bIsLoaded;
+}
+
+FString FTensorShapeProtoDimension::ToString(const FString& InLineStarted) const
+{
+	FString String = TEXT("\n");
+	FModelProtoStringParser::PositiveInt32Or64ToStringAuxiliary(String, TEXT("{0}FTensorShapeProtoDimension.DimValue: {1}\n"), InLineStarted, DimValue);
+	FModelProtoStringParser::StringToStringIfNotEmptyAuxiliary(String, TEXT("{0}FTensorShapeProtoDimension.DimParam: {1}\n"), InLineStarted, DimParam);
+	return String;
+}
+
+
+
+/* FTensorShapeProto public functions
+ *****************************************************************************/
+
+bool FTensorShapeProto::LoadFromString(const FString& InProtoString, const int32 InLevel)
+{
+	// Sanity check
+	if (InProtoString.Len() < 1)
+	{
+		return false;
+	}
+
+	// FString text into [key, value] maps
+	const TMap<FString, TArray<FString>> ProtoMap = FModelProtoStringParser::ModelProtoStringToMap(InProtoString, InLevel);
+
+	// Fill class variables
+	const bool bIsLoaded = FModelProtoStringParser::GetModelProtoArray(Dim, ProtoMap, TEXT("dim"), InLevel + 1);
+
+	if (!bIsLoaded)
+	{
+		UE_LOG(LogModelProto, Warning, TEXT("FTensorShapeProto::LoadFromString() failed."));
+	}
+	return bIsLoaded;
+}
+
+FString FTensorShapeProto::ToString(const FString& InLineStarted) const
+{
+	FString String = TEXT("\n");
+	FModelProtoStringParser::ArrayTToStringAuxiliary(String, TEXT("{0}FTensorShapeProto.Dim[{1}]: {2}"), InLineStarted, Dim, FModelProtoStringParser::ToStringSeparator + InLineStarted);
+	return String;
+}
+
+
+
+/* FTensorProto structors
+ *****************************************************************************/
+
+FTensorProto::FTensorProto()
+	: DataType(ETensorProtoDataType::UNDEFINED)
+	, DataLocation(ETensorProtoDataLocation::DEFAULT)
+{
+}
+
+/* FTensorProto public functions
+ *****************************************************************************/
+
+bool FTensorProto::LoadFromString(const FString& InProtoString, const int32 InLevel)
+{
+	// Sanity check
+	if (InProtoString.Len() < 1)
+	{
+		return false;
+	}
+
+	// FString text into [key, value] maps
+	const TMap<FString, TArray<FString>> ProtoMap = FModelProtoStringParser::ModelProtoStringToMap(InProtoString, InLevel);
+
+	// Fill class variables
+	bool bIsLoaded = true;
+	Dimensions = FModelProtoStringParser::GetModelProtoInt64Array(ProtoMap, TEXT("dims"));
+	DataType = StringToDataType(FModelProtoStringParser::GetModelProtoStringOrEmpty(ProtoMap, TEXT("data_type")));
+	DataLocation = StringToDataLocation(FModelProtoStringParser::GetModelProtoStringOrEmpty(ProtoMap, TEXT("data_location")));
+	/*bIsLoaded &= */Segment.LoadFromString(FModelProtoStringParser::GetModelProtoStringOrEmpty(ProtoMap, TEXT("segment")), InLevel + 1);
+	FloatData = FModelProtoStringParser::GetModelProtoFloatArray(ProtoMap, TEXT("float_data"));
+	Int32Data = FModelProtoStringParser::GetModelProtoInt32Array(ProtoMap, TEXT("int32_data"));
+	StringData = FModelProtoStringParser::GetModelProtoStringArray(ProtoMap, TEXT("string_data"));
+	Int64Data = FModelProtoStringParser::GetModelProtoInt64Array(ProtoMap, TEXT("int64_data"));
+	Name = FModelProtoStringParser::GetModelProtoStringOrEmpty(ProtoMap, TEXT("name"));
+	DocString = FModelProtoStringParser::GetModelProtoStringOrEmpty(ProtoMap, TEXT("doc_string"));
+	RawData = FModelProtoStringParser::GetModelProtoCharAsUInt8Array(ProtoMap, TEXT("raw_data"));
+	bIsLoaded &= FModelProtoStringParser::GetModelProtoArray(ExternalData, ProtoMap, TEXT("external_data"), InLevel + 1);
+	DoubleData = FModelProtoStringParser::GetModelProtoDoubleArray(ProtoMap, TEXT("double_data"));
+	UInt64Data = FModelProtoStringParser::GetModelProtoUInt64Array(ProtoMap, TEXT("uint64_data"));
+
+	bIsLoaded &= (bIsLoaded && DataType != ETensorProtoDataType::UNDEFINED
+		&& ((ExternalData.Num() == 1 && ExternalData[0].Key == TEXT("location") && ExternalData[0].Value.Len() > 0) // Either external data is available...
+			|| !FloatData.IsEmpty() || !Int32Data.IsEmpty() || !Int64Data.IsEmpty() || !RawData.IsEmpty() || !DoubleData.IsEmpty())); // ...or RawData is filled
+	if (!bIsLoaded)
+	{
+		UE_LOG(LogModelProto, Warning, TEXT("FTensorProto::LoadFromString() failed."));
+	}
+	return bIsLoaded;
+}
+
+FString FTensorProto::ToString(const FString& InLineStarted) const
+{
+	FString String = TEXT("\n");
+	FModelProtoStringParser::GreaterThan0Int32Or64ToStringAuxiliary(String, TEXT("{0}FTensorProto.Dimensions.Num(): {1}\n"), InLineStarted, Dimensions.Num());
+	String += FString::Format(TEXT("{0}FTensorProto.DataType: {1}\n"), { InLineStarted, FString::FromInt((int32)DataType) });
+	String += FString::Format(TEXT("{0}FTensorProto.DataLocation: {1}\n"), { InLineStarted, FString::FromInt((int32)DataLocation) });
+	if (Segment.Begin > -1 || Segment.End > -1)
+	{
+		String += FString::Format(TEXT("{0}FTensorProto.Segment: {1}"), { InLineStarted, Segment.ToString(FModelProtoStringParser::ToStringSeparator + InLineStarted) });
+	}
+	FModelProtoStringParser::GreaterThan0Int32Or64ToStringAuxiliary(String, TEXT("{0}FTensorProto.FloatData.Num(): {1}\n"), InLineStarted, FloatData.Num());
+	FModelProtoStringParser::GreaterThan0Int32Or64ToStringAuxiliary(String, TEXT("{0}FTensorProto.Int32Data.Num(): {1}\n"), InLineStarted, Int32Data.Num());
+	FModelProtoStringParser::GreaterThan0Int32Or64ToStringAuxiliary(String, TEXT("{0}FTensorProto.StringData.Num(): {1}\n"), InLineStarted, StringData.Num());
+	FModelProtoStringParser::GreaterThan0Int32Or64ToStringAuxiliary(String, TEXT("{0}FTensorProto.Int64Data.Num(): {1}\n"), InLineStarted, Int64Data.Num());
+	FModelProtoStringParser::StringToStringIfNotEmptyAuxiliary(String, TEXT("{0}FTensorProto.Name: {1}\n"), InLineStarted, Name);
+	FModelProtoStringParser::GreaterThan0Int32Or64ToStringAuxiliary(String, TEXT("{0}FTensorProto.RawData.Num(): {1}\n"), InLineStarted, RawData.Num());
+	FModelProtoStringParser::ArrayTToStringAuxiliary(String, TEXT("{0}FTensorProto.ExternalData[{1}]: {2}"), InLineStarted, ExternalData, FModelProtoStringParser::ToStringSeparator + InLineStarted);
+	FModelProtoStringParser::GreaterThan0Int32Or64ToStringAuxiliary(String, TEXT("{0}FTensorProto.DoubleData.Num(): {1}\n"), InLineStarted, DoubleData.Num());
+	FModelProtoStringParser::GreaterThan0Int32Or64ToStringAuxiliary(String, TEXT("{0}FTensorProto.UInt64Data.Num(): {1}\n"), InLineStarted, UInt64Data.Num());
+	return String;
+}
+
+ENeuralDataType FTensorProto::GetDataTypeFromTensorProtoDataType() const
+{
+	if (DataType == ETensorProtoDataType::FLOAT)
+	{
+		return ENeuralDataType::Float;
+	}
+	else if (DataType == ETensorProtoDataType::INT32)
+	{
+		return ENeuralDataType::Int32;
+	}
+	else if (DataType == ETensorProtoDataType::INT64)
+	{
+		return ENeuralDataType::Int64;
+	}
+	else if (DataType == ETensorProtoDataType::UINT32)
+	{
+		return ENeuralDataType::UInt32;
+	}
+	else if (DataType == ETensorProtoDataType::UINT64)
+	{
+		return ENeuralDataType::UInt64;
+	}
+	else if (DataType == ETensorProtoDataType::UNDEFINED)
+	{
+		return ENeuralDataType::None;
+	}
+	UE_LOG(LogModelProto, Warning, TEXT("FModelProto::GetDataTypeFromTensorProtoDataType(): Unknown DataType = %d used."), (int32)DataType);
+	return ENeuralDataType::None;
+}
+
+/* FTensorProto private functions
+ *****************************************************************************/
+
+ETensorProtoDataType FTensorProto::StringToDataType(const FString& InString)
+{
+	if (InString == TEXT("UNDEFINED") || InString == TEXT("0"))
+	{
+		return ETensorProtoDataType::UNDEFINED;
+	}
+	else if (InString == TEXT("FLOAT") || InString == TEXT("1"))
+	{
+		return ETensorProtoDataType::FLOAT;
+	}
+	else if (InString == TEXT("UINT8") || InString == TEXT("2"))
+	{
+		return ETensorProtoDataType::UINT8;
+	}
+	else if (InString == TEXT("INT8") || InString == TEXT("3"))
+	{
+		return ETensorProtoDataType::INT8;
+	}
+	else if (InString == TEXT("UINT16") || InString == TEXT("4"))
+	{
+		return ETensorProtoDataType::UINT16;
+	}
+	else if (InString == TEXT("INT16") || InString == TEXT("5"))
+	{
+		return ETensorProtoDataType::INT16;
+	}
+	else if (InString == TEXT("INT32") || InString == TEXT("6"))
+	{
+		return ETensorProtoDataType::INT32;
+	}
+	else if (InString == TEXT("INT64") || InString == TEXT("7"))
+	{
+		return ETensorProtoDataType::INT64;
+	}
+	else if (InString == TEXT("STRING") || InString == TEXT("8"))
+	{
+		return ETensorProtoDataType::STRING;
+	}
+	else if (InString == TEXT("BOOL") || InString == TEXT("9"))
+	{
+		return ETensorProtoDataType::BOOL;
+	}
+	else if (InString == TEXT("FLOAT16") || InString == TEXT("10"))
+	{
+		return ETensorProtoDataType::FLOAT16;
+	}
+	else if (InString == TEXT("DOUBLE") || InString == TEXT("11"))
+	{
+		return ETensorProtoDataType::DOUBLE;
+	}
+	else if (InString == TEXT("UINT32") || InString == TEXT("12"))
+	{
+		return ETensorProtoDataType::UINT32;
+	}
+	else if (InString == TEXT("UINT64") || InString == TEXT("13"))
+	{
+		return ETensorProtoDataType::UINT64;
+	}
+	else if (InString == TEXT("COMPLEX64") || InString == TEXT("14"))
+	{
+		return ETensorProtoDataType::COMPLEX64;
+	}
+	else if (InString == TEXT("COMPLEX128") || InString == TEXT("15"))
+	{
+		return ETensorProtoDataType::COMPLEX128;
+	}
+	else if (InString == TEXT("BFLOAT16") || InString == TEXT("16"))
+	{
+		return ETensorProtoDataType::BFLOAT16;
+	}
+	else
+	{
+		UE_LOG(LogModelProto, Warning,
+			TEXT("FTensorProto::StringToDataType(): Not implemented yet for string %s. Notify us."), *InString);
+		return ETensorProtoDataType::UNDEFINED;
+	}
+}
+
+ETensorProtoDataLocation FTensorProto::StringToDataLocation(const FString& InString) 
+{
+	if (InString == TEXT("DEFAULT") || InString == TEXT("0") || InString == TEXT(""))
+	{
+		return ETensorProtoDataLocation::DEFAULT;
+	}
+	else if (InString == TEXT("EXTERNAL") || InString == TEXT("1"))
+	{
+		return ETensorProtoDataLocation::EXTERNAL;
+	}
+	else
+	{
+		UE_LOG(LogModelProto, Warning,
+			TEXT("FTensorProto::StringToDataLocation(): Not implemented yet for string %s. Notify us."), *InString);
+		return ETensorProtoDataLocation::DEFAULT;
+	}
+}
+
+
+
+/* FTypeProtoTensor structors
+ *****************************************************************************/
+
+FTypeProtoTensor::FTypeProtoTensor()
+	: ElemType(-1)
+{
+}
+
+/* FTypeProtoTensor public functions
+ *****************************************************************************/
+
+bool FTypeProtoTensor::LoadFromString(const FString& InProtoString, const int32 InLevel)
+{
+	// Sanity check
+	if (InProtoString.Len() < 1)
+	{
+		return false;
+	}
+
+	// FString text into [key, value] maps
+	const TMap<FString, TArray<FString>> ProtoMap = FModelProtoStringParser::ModelProtoStringToMap(InProtoString, InLevel);
+
+	// Fill class variables
+	ElemType = FModelProtoStringParser::GetModelProtoInt64(ProtoMap, TEXT("elem_type"));
+	/*bIsLoaded &= */Shape.LoadFromString(FModelProtoStringParser::GetModelProtoStringOrEmpty(ProtoMap, TEXT("shape")), InLevel + 1);
+
+	const bool bIsLoaded = ElemType > -1;
+	if (!bIsLoaded)
+	{
+		UE_LOG(LogModelProto, Warning, TEXT("FTypeProtoTensor::LoadFromString() failed."));
+	}
+	return bIsLoaded;
+}
+
+FString FTypeProtoTensor::ToString(const FString& InLineStarted) const
+{
+	FString String = TEXT("\n");
+	FModelProtoStringParser::PositiveInt32Or64ToStringAuxiliary(String, TEXT("{0}FTypeProtoTensor.ElemType: {1}\n"), InLineStarted, ElemType);
+	String += FString::Format(TEXT("{0}FTypeProtoTensor.Shape: {1}"), { InLineStarted, Shape.ToString(FModelProtoStringParser::ToStringSeparator + InLineStarted) });
+	return String;
+}
+
+
+
+/* FTypeProtoMap structors
+ *****************************************************************************/
+
+FTypeProtoMap::FTypeProtoMap()
+	: KeyType(0)
+{
+}
+
+
+
+/* FSparseTensorProto public functions
+ *****************************************************************************/
+
+bool FSparseTensorProto::LoadFromString(const FString& InProtoString, const int32 InLevel)
+{
+	// Sanity check
+	if (InProtoString.Len() < 1)
+	{
+		return false;
+	}
+
+	// FString text into [key, value] maps
+	const TMap<FString, TArray<FString>> ProtoMap = FModelProtoStringParser::ModelProtoStringToMap(InProtoString, InLevel);
+
+	// Fill class variables
+	bool bIsLoaded = true;
+	const FString ValuesString = FModelProtoStringParser::GetModelProtoStringOrEmpty(ProtoMap, TEXT("values"));
+	if (!ValuesString.IsEmpty())
+	{
+		bIsLoaded &= Values.LoadFromString(ValuesString, InLevel + 1);
+	}
+	const FString IndicesString = FModelProtoStringParser::GetModelProtoStringOrEmpty(ProtoMap, TEXT("indices"));
+	if (!IndicesString.IsEmpty())
+	{
+		bIsLoaded &= Indices.LoadFromString(IndicesString, InLevel + 1);
+	}
+	Dimensions = FModelProtoStringParser::GetModelProtoInt64Array(ProtoMap, TEXT("dims"));
+
+
+	if (!bIsLoaded)
+	{
+		UE_LOG(LogModelProto, Warning, TEXT("FSparseTensorProto::LoadFromString() failed."));
+	}
+	return bIsLoaded;
+
+}
+
+FString FSparseTensorProto::ToString(const FString& InLineStarted) const
+{
+	FString String = TEXT("\n");
+	String += FString::Format(TEXT("{0}FSparseTensorProto.Values: {1}"), { InLineStarted, Values.ToString(FModelProtoStringParser::ToStringSeparator + InLineStarted) });
+	String += FString::Format(TEXT("{0}FSparseTensorProto.Indices: {1}"), { InLineStarted, Indices.ToString(FModelProtoStringParser::ToStringSeparator + InLineStarted) });
+	FModelProtoStringParser::GreaterThan0Int32Or64ToStringAuxiliary(String, TEXT("{0}FSparseTensorProto.Dimensions.Num(): {1}\n"), InLineStarted, Dimensions.Num());
+	return String;
+}
+
+
+
+/* FTypeProto public functions
+ *****************************************************************************/
+
+bool FTypeProto::LoadFromString(const FString& InProtoString, const int32 InLevel)
+{
+	// Sanity check
+	if (InProtoString.Len() < 1)
+	{
+		return false;
+	}
+
+	// FString text into [key, value] maps
+	const TMap<FString, TArray<FString>> ProtoMap = FModelProtoStringParser::ModelProtoStringToMap(InProtoString, InLevel);
+
+	// Fill class variables
+	const bool bIsLoaded = TensorType.LoadFromString(FModelProtoStringParser::GetModelProtoStringOrEmpty(ProtoMap, TEXT("tensor_type")), InLevel + 1);
+	Denotation = FModelProtoStringParser::GetModelProtoStringOrEmpty(ProtoMap, TEXT("denotation"));
+
+	if (!bIsLoaded)
+	{
+		UE_LOG(LogModelProto, Warning, TEXT("FTypeProto::LoadFromString() failed."));
+	}
+	return bIsLoaded;
+}
+
+FString FTypeProto::ToString(const FString& InLineStarted) const
+{
+	FString String = TEXT("\n");
+	String += FString::Format(TEXT("{0}FTypeProto.TensorType: {1}"), { InLineStarted, TensorType.ToString(FModelProtoStringParser::ToStringSeparator + InLineStarted) });
+	FModelProtoStringParser::StringToStringIfNotEmptyAuxiliary(String, TEXT("{0}FTypeProto.Denotation: {1}\n"), InLineStarted, Denotation);
+	return String;
+}
+
+
+
+/* FAttributeProto structors
+ *****************************************************************************/
+
+FAttributeProto::FAttributeProto()
+	: Type(EAttributeProtoAttributeType::UNDEFINED)
+	, F(-1.f)
+	, I(-1)
+{
+}
+
+/* FAttributeProto public functions
+ *****************************************************************************/
+
+bool FAttributeProto::LoadFromString(const FString& InProtoString, const int32 InLevel)
+{
+	// Sanity check
+	if (InProtoString.Len() < 1)
+	{
+		return false;
+	}
+
+	// FString text into [key, value] maps
+	const TMap<FString, TArray<FString>> ProtoMap = FModelProtoStringParser::ModelProtoStringToMap(InProtoString, InLevel);
+
+	// Fill class variables
+	bool bIsLoaded = true;
+	DocString = FModelProtoStringParser::GetModelProtoStringOrEmpty(ProtoMap, TEXT("doc_string"));
+	F = FModelProtoStringParser::GetModelProtoFloat(ProtoMap, TEXT("f"));
+	Floats = FModelProtoStringParser::GetModelProtoFloatArray(ProtoMap, TEXT("floats"));
+	if (ProtoMap.Find(TEXT("floats")))
+	{
+		UE_LOG(LogModelProto, Warning, TEXT("Not implemented for this field (\"floats\") yet. ProtoString = %s."), *InProtoString);
+	}
+	if (ProtoMap.Find(TEXT("g")))
+	{
+		UE_LOG(LogModelProto, Warning, TEXT("Not implemented for this field (\"g\") yet. ProtoString = %s."), *InProtoString);
+	}
+// 	const FString GString = FModelProtoStringParser::GetModelProtoStringOrEmpty(ProtoMap, TEXT("g"));
+// 	if (!GString.IsEmpty())
+// 	{
+// 		G = MakeShared<FGraphProto>();
+// 		/*bIsLoaded &= */G->LoadFromString(
+// 		FModelProtoStringParser::GetModelProtoStringOrEmpty(ProtoMap, TEXT("g")), InLevel + 1);
+// 	}
+// 	if (ProtoMap.Find(TEXT("graphs")))
+// 	{
+// 		UE_LOG(LogModelProto, Warning, TEXT("Not implemented for this field (\"graphs\") yet. ProtoString = %s."), *InProtoString);
+// 	}
+// 	/*bIsLoaded &= */Graphs.LoadFromString(FModelProtoStringParser::GetModelProtoStringOrEmpty(ProtoMap, TEXT("graphs")), InLevel + 1);
+	I = FModelProtoStringParser::GetModelProtoInt64(ProtoMap, TEXT("i"));
+	Integers = FModelProtoStringParser::GetModelProtoInt64Array(ProtoMap, TEXT("ints"));
+	Name = FModelProtoStringParser::GetModelProtoStringOrEmpty(ProtoMap, TEXT("name"));
+	S = FModelProtoStringParser::GetModelProtoStringOrEmpty(ProtoMap, TEXT("s"));
+	Strings = FModelProtoStringParser::GetModelProtoStringArray(ProtoMap, TEXT("strings"));
+	const FString TString = FModelProtoStringParser::GetModelProtoStringOrEmpty(ProtoMap, TEXT("t"));
+	if (!TString.IsEmpty())
+	{
+		bIsLoaded &= T.LoadFromString(TString, InLevel + 1);
+	}
+	bIsLoaded &= FModelProtoStringParser::GetModelProtoArray(Tensors, ProtoMap, TEXT("tensors"), InLevel + 1);
+	const FString STString = FModelProtoStringParser::GetModelProtoStringOrEmpty(ProtoMap, TEXT("sparse_tensor"));
+	if (!STString.IsEmpty())
+	{
+		bIsLoaded &= SparseTensor.LoadFromString(STString, InLevel + 1);
+	}
+	bIsLoaded &= FModelProtoStringParser::GetModelProtoArray(SparseTensors, ProtoMap, TEXT("sparse_tensors"), InLevel + 1);
+	Type = StringToAttributeType(FModelProtoStringParser::GetModelProtoStringOrEmpty(ProtoMap, TEXT("type")));
+
+	bIsLoaded &= (!Name.IsEmpty() && Type != EAttributeProtoAttributeType::UNDEFINED);
+	if (!bIsLoaded)
+	{
+		UE_LOG(LogModelProto, Warning, TEXT("FAttributeProto::LoadFromString() failed."));
+	}
+	return bIsLoaded;
+}
+
+FString FAttributeProto::ToString(const FString& InLineStarted) const
+{
+	FString String = TEXT("\n");
+	String += FString::Format(TEXT("{0}FAttributeProto.Type: {1}\n"), { InLineStarted, FString::FromInt((int32)Type) });
+	FModelProtoStringParser::StringToStringIfNotEmptyAuxiliary(String, TEXT("{0}FAttributeProto.Name: {1}\n"), InLineStarted, Name);
+	FModelProtoStringParser::StringToStringIfNotEmptyAuxiliary(String, TEXT("{0}FAttributeProto.DocString: {1}\n"), InLineStarted, DocString);
+	if (Type != EAttributeProtoAttributeType::UNDEFINED)
+	{
+		if (Type == EAttributeProtoAttributeType::FLOAT)
+		{
+			String += FString::Format(TEXT("{0}FAttributeProto.F: {1}\n"), { InLineStarted, FString::SanitizeFloat(F) });
+		}
+		else if (Type == EAttributeProtoAttributeType::INT)
+		{
+			FModelProtoStringParser::PositiveInt32Or64ToStringAuxiliary(String, TEXT("{0}FAttributeProto.I: {1}\n"), InLineStarted, I);
+		}
+		else if (Type == EAttributeProtoAttributeType::STRING)
+		{
+			FModelProtoStringParser::StringToStringIfNotEmptyAuxiliary(String, TEXT("{0}FAttributeProto.S: {1}\n"), InLineStarted, S);
+		}
+		else if (Type == EAttributeProtoAttributeType::TENSOR)
+		{
+			String += FString::Format(TEXT("{0}FAttributeProto.T: {1}"), { InLineStarted, T.ToString(FModelProtoStringParser::ToStringSeparator + InLineStarted) });
+		}
+		//else if (Type == EAttributeProtoAttributeType::GRAPH)
+		//{
+		//	if (G.IsValid())
+		//	{
+		//		String += FString::Format(TEXT("{0}FAttributeProto.G: {1}"), { InLineStarted, G->ToString(FModelProtoStringParser::ToStringSeparator + InLineStarted) });
+		//	}
+		//}
+		else if (Type == EAttributeProtoAttributeType::SPARSE_TENSOR)
+		{
+			String += FString::Format(TEXT("{0}FAttributeProto.SparseTensor: {1}"), { InLineStarted, SparseTensor.ToString(FModelProtoStringParser::ToStringSeparator + InLineStarted) });
+		}
+		else if (Type == EAttributeProtoAttributeType::FLOATS)
+		{
+			FModelProtoStringParser::GreaterThan0Int32Or64ToStringAuxiliary(String, TEXT("{0}FAttributeProto.Floats.Num(): {1}\n"), InLineStarted, Floats.Num());
+		}
+		else if (Type == EAttributeProtoAttributeType::INTS)
+		{
+			FModelProtoStringParser::GreaterThan0Int32Or64ToStringAuxiliary(String, TEXT("{0}FAttributeProto.Integers.Num(): {1}\n"), InLineStarted, Integers.Num());
+		}
+		else if (Type == EAttributeProtoAttributeType::STRINGS)
+		{
+			FModelProtoStringParser::GreaterThan0Int32Or64ToStringAuxiliary(String, TEXT("{0}FAttributeProto.Strings.Num(): {1}\n"), InLineStarted, Strings.Num());
+		}
+		else if (Type == EAttributeProtoAttributeType::TENSORS)
+		{
+			FModelProtoStringParser::GreaterThan0Int32Or64ToStringAuxiliary(String, TEXT("{0}FAttributeProto.Tensors.Num(): {1}\n"), InLineStarted, Tensors.Num());
+		}
+		//else if (Type == EAttributeProtoAttributeType::GRAPHS)
+		//{
+		//	FModelProtoStringParser::GreaterThan0Int32Or64ToStringAuxiliary(String, TEXT("{0}FAttributeProto.Graphs.Num(): {1}\n"), InLineStarted, Graphs.Num());
+		//}
+		else if (Type == EAttributeProtoAttributeType::SPARSE_TENSORS)
+		{
+			FModelProtoStringParser::GreaterThan0Int32Or64ToStringAuxiliary(String, TEXT("{0}FAttributeProto.SparseTensors.Num(): {1}\n"), InLineStarted, SparseTensors.Num());
+		}
+	}
+	return String;
+}
+
+/* FAttributeProto private functions
+ *****************************************************************************/
+
+EAttributeProtoAttributeType FAttributeProto::StringToAttributeType(const FString& InString)
+{
+	if (InString == TEXT("UNDEFINED"))
+	{
+		return EAttributeProtoAttributeType::UNDEFINED;
+	}
+	else if (InString == TEXT("FLOAT"))
+	{
+		return EAttributeProtoAttributeType::FLOAT;
+	}
+	else if (InString == TEXT("INT"))
+	{
+		return EAttributeProtoAttributeType::INT;
+	}
+	else if (InString == TEXT("STRING"))
+	{
+		return EAttributeProtoAttributeType::STRING;
+	}
+	else if (InString == TEXT("TENSOR"))
+	{
+		return EAttributeProtoAttributeType::TENSOR;
+	}
+	else if (InString == TEXT("GRAPH"))
+	{
+		return EAttributeProtoAttributeType::GRAPH;
+	}
+	else if (InString == TEXT("SPARSE_TENSOR"))
+	{
+		return EAttributeProtoAttributeType::SPARSE_TENSOR;
+	}
+	else if (InString == TEXT("FLOATS"))
+	{
+		return EAttributeProtoAttributeType::FLOATS;
+	}
+	else if (InString == TEXT("INTS"))
+	{
+		return EAttributeProtoAttributeType::INTS;
+	}
+	else if (InString == TEXT("STRINGS"))
+	{
+		return EAttributeProtoAttributeType::STRINGS;
+	}
+	else if (InString == TEXT("TENSORS"))
+	{
+		return EAttributeProtoAttributeType::TENSORS;
+	}
+	else if (InString == TEXT("GRAPHS"))
+	{
+		return EAttributeProtoAttributeType::GRAPHS;
+	}
+	else if (InString == TEXT("SPARSE_TENSORS"))
+	{
+		return EAttributeProtoAttributeType::SPARSE_TENSORS;
+	}
+	else
+	{
+		UE_LOG(LogModelProto, Warning,
+			TEXT("FAttributeProto::StringToAttributeType(): Not implemented yet for string %s. Notify us."), *InString);
+		return EAttributeProtoAttributeType::UNDEFINED;
+	}
+}
+
+
+
+/* FValueInfoProto public functions
+ *****************************************************************************/
+
+bool FValueInfoProto::LoadFromString(const FString& InProtoString, const int32 InLevel)
+{
+	// Sanity check
+	if (InProtoString.Len() < 1)
+	{
+		return false;
+	}
+
+	// FString text into [key, value] maps
+	const TMap<FString, TArray<FString>> ProtoMap = FModelProtoStringParser::ModelProtoStringToMap(InProtoString, InLevel);
+
+	// Fill class variables
+	bool bIsLoaded = true;
+	Name = FModelProtoStringParser::GetModelProtoStringOrEmpty(ProtoMap, TEXT("name"));
+	bIsLoaded &= Type.LoadFromString(FModelProtoStringParser::GetModelProtoStringOrEmpty(ProtoMap, TEXT("type")), InLevel + 1);
+	DocString = FModelProtoStringParser::GetModelProtoStringOrEmpty(ProtoMap, TEXT("doc_string"));
+
+	bIsLoaded &= !Name.IsEmpty();
+	if (!bIsLoaded)
+	{
+		UE_LOG(LogModelProto, Warning, TEXT("FValueInfoProto::LoadFromString() failed."));
+	}
+	return bIsLoaded;
+}
+
+FString FValueInfoProto::ToString(const FString& InLineStarted) const
+{
+	FString String = TEXT("\n");
+	FModelProtoStringParser::StringToStringIfNotEmptyAuxiliary(String, TEXT("{0}FValueInfoProto.Name: {1}\n"), InLineStarted, Name);
+	String += FString::Format(TEXT("{0}FValueInfoProto.Type: {1}"), { InLineStarted, Type.ToString(FModelProtoStringParser::ToStringSeparator + InLineStarted) });
+	FModelProtoStringParser::StringToStringIfNotEmptyAuxiliary(String, TEXT("{0}FValueInfoProto.DocString: {1}\n"), InLineStarted, DocString);
+	return String;
+}
+
+
+
+/* FTensorAnnotation public functions
+ *****************************************************************************/
+
+bool FTensorAnnotation::LoadFromString(const FString& InProtoString, const int32 InLevel)
+{
+	// Sanity check
+	if (InProtoString.Len() < 1)
+	{
+		return false;
+	}
+
+	// FString text into [key, value] maps
+	const TMap<FString, TArray<FString>> ProtoMap = FModelProtoStringParser::ModelProtoStringToMap(InProtoString, InLevel);
+
+	// Fill class variables
+	bool bIsLoaded = true;
+	TensorName = FModelProtoStringParser::GetModelProtoStringOrEmpty(ProtoMap, TEXT("tensor_name"));
+	bIsLoaded &= FModelProtoStringParser::GetModelProtoArray(QuantParameterTensorNames, ProtoMap, TEXT("quant_parameter_tensor_names"), InLevel + 1);
+
+	bIsLoaded &= !TensorName.IsEmpty();
+	if (!bIsLoaded)
+	{
+		UE_LOG(LogModelProto, Warning, TEXT("FTensorAnnotation::LoadFromString() failed."));
+	}
+	return bIsLoaded;
+}
+
+FString FTensorAnnotation::ToString(const FString& InLineStarted) const
+{
+	FString String = TEXT("\n");
+	FModelProtoStringParser::StringToStringIfNotEmptyAuxiliary(String, TEXT("{0}FTensorAnnotation.TensorName: {1}\n"), InLineStarted, TensorName);
+	FModelProtoStringParser::ArrayTToStringAuxiliary(String, TEXT("{0}FTensorAnnotation.QuantParmeterTensorNames[{1}]: {2}"), InLineStarted, QuantParameterTensorNames, FModelProtoStringParser::ToStringSeparator + InLineStarted);
+	return String;
+}
+
+
+
+/* FNodeProto public functions
+ *****************************************************************************/
+
+bool FNodeProto::LoadFromString(const FString& InProtoString, const int32 InLevel)
+{
+	// Sanity check
+	if (InProtoString.Len() < 1)
+	{
+		return false;
+	}
+
+	// FString text into [key, value] maps
+	const TMap<FString, TArray<FString>> ProtoMap = FModelProtoStringParser::ModelProtoStringToMap(InProtoString, InLevel);
+
+	// Fill class variables
+	bool bIsLoaded = true;
+	Input = FModelProtoStringParser::GetModelProtoStringArray(ProtoMap, TEXT("input"));
+	Output = FModelProtoStringParser::GetModelProtoStringArray(ProtoMap, TEXT("output"));
+	Name = FModelProtoStringParser::GetModelProtoStringOrEmpty(ProtoMap, TEXT("name"));
+	OperatorType = FModelProtoStringParser::GetModelProtoStringOrEmpty(ProtoMap, TEXT("op_type"));
+	Domain = FModelProtoStringParser::GetModelProtoStringOrEmpty(ProtoMap, TEXT("domain"));
+	bIsLoaded &= FModelProtoStringParser::GetModelProtoArray(Attribute, ProtoMap, TEXT("attribute"), InLevel + 1);
+	DocString = FModelProtoStringParser::GetModelProtoStringOrEmpty(ProtoMap, TEXT("doc_string"));
+
+	bIsLoaded &= !OperatorType.IsEmpty();
+	if (!bIsLoaded)
+	{
+		UE_LOG(LogModelProto, Warning, TEXT("FNodeProto::LoadFromString() failed."));
+	}
+	return bIsLoaded;
+}
+
+FString FNodeProto::ToString(const FString& InLineStarted) const
+{
+	FString String = TEXT("\n");
+	FModelProtoStringParser::StringArrayToStringAuxiliary(String, TEXT("{0}FNodeProto.Input: [{1}]\n"), InLineStarted, Input);
+	FModelProtoStringParser::StringArrayToStringAuxiliary(String, TEXT("{0}FNodeProto.Output: [{1}]\n"), InLineStarted, Output);
+	FModelProtoStringParser::StringToStringIfNotEmptyAuxiliary(String, TEXT("{0}FNodeProto.Name: {1}\n"), InLineStarted, Name);
+	FModelProtoStringParser::StringToStringIfNotEmptyAuxiliary(String, TEXT("{0}FNodeProto.OperatorType: {1}\n"), InLineStarted, OperatorType);
+	FModelProtoStringParser::StringToStringIfNotEmptyAuxiliary(String, TEXT("{0}FNodeProto.Domain: {1}\n"), InLineStarted, Domain);
+	FModelProtoStringParser::ArrayTToStringAuxiliary(String, TEXT("{0}FNodeProto.Attribute[{1}]: {2}"), InLineStarted, Attribute, FModelProtoStringParser::ToStringSeparator + InLineStarted);
+	FModelProtoStringParser::StringToStringIfNotEmptyAuxiliary(String, TEXT("{0}FNodeProto.DocString: {1}\n"), InLineStarted, DocString);
+	return String;
+}
+
+
+
+/* FGraphProto structors
+ *****************************************************************************/
+
+FGraphProto::FGraphProto()
+	: bIsLoaded(false)
+{
+}
+
+/* FGraphProto public functions
+ *****************************************************************************/
+
+bool FGraphProto::LoadFromString(const FString& InProtoString, const int32 InLevel)
+{
+	// Sanity check
+	if (InProtoString.Len() < 1)
+	{
+		return false;
+	}
+
+	// FString text into [key, value] maps
+	const TMap<FString, TArray<FString>> ProtoMap = FModelProtoStringParser::ModelProtoStringToMap(InProtoString, InLevel);
+
+	// Fill class variables
+	bIsLoaded = true;
+	bIsLoaded &= FModelProtoStringParser::GetModelProtoArray(Node, ProtoMap, TEXT("node"), InLevel + 1);
+	Name = FModelProtoStringParser::GetModelProtoStringOrEmpty(ProtoMap, TEXT("name"));
+	bIsLoaded &= FModelProtoStringParser::GetModelProtoArray(Initializer, ProtoMap, TEXT("initializer"), InLevel + 1);
+	bIsLoaded &= FModelProtoStringParser::GetModelProtoArray(SparseInitializer, ProtoMap, TEXT("sparse_initializer"), InLevel + 1);
+	DocString = FModelProtoStringParser::GetModelProtoStringOrEmpty(ProtoMap, TEXT("doc_string"));
+	bIsLoaded &= FModelProtoStringParser::GetModelProtoArray(Input, ProtoMap, TEXT("input"), InLevel + 1);
+	bIsLoaded &= FModelProtoStringParser::GetModelProtoArray(Output, ProtoMap, TEXT("output"), InLevel + 1);
+	bIsLoaded &= FModelProtoStringParser::GetModelProtoArray(ValueInfo, ProtoMap, TEXT("value_info"), InLevel + 1);
+	bIsLoaded &= FModelProtoStringParser::GetModelProtoArray(QuantizationAnnotation, ProtoMap, TEXT("quantization_annotation"), InLevel + 1);
+
+	bIsLoaded &= !Name.IsEmpty();
+	if (!bIsLoaded)
+	{
+		UE_LOG(LogModelProto, Warning, TEXT("FGraphProto::LoadFromString() failed."));
+	}
+	return bIsLoaded;
+}
+
+FString FGraphProto::ToString(const FString& InLineStarted) const
+{
+	if (bIsLoaded)
+	{
+		FString String = TEXT("\n");
+		FModelProtoStringParser::ArrayTToStringAuxiliary(String, TEXT("{0}FGraphProto.Node[{1}]: {2}"), InLineStarted, Node, FModelProtoStringParser::ToStringSeparator + InLineStarted);
+		FModelProtoStringParser::StringToStringIfNotEmptyAuxiliary(String, TEXT("{0}FGraphProto.Name: {1}\n"), InLineStarted, Name);
+		FModelProtoStringParser::ArrayTToStringAuxiliary(String, TEXT("{0}FGraphProto.Initializer[{1}]: {2}"), InLineStarted, Initializer, FModelProtoStringParser::ToStringSeparator + InLineStarted);
+		FModelProtoStringParser::GreaterThan0Int32Or64ToStringAuxiliary(String, TEXT("{0}FGraphProto.SparseInitializer.Num(): {1}\n"), InLineStarted, SparseInitializer.Num());
+		FModelProtoStringParser::StringToStringIfNotEmptyAuxiliary(String, TEXT("{0}FGraphProto.DocString: {1}\n"), InLineStarted, DocString);
+		FModelProtoStringParser::ArrayTToStringAuxiliary(String, TEXT("{0}FGraphProto.Input[{1}]: {2}"), InLineStarted, Input, FModelProtoStringParser::ToStringSeparator + InLineStarted);
+		FModelProtoStringParser::ArrayTToStringAuxiliary(String, TEXT("{0}FGraphProto.Output[{1}]: {2}"), InLineStarted, Output, FModelProtoStringParser::ToStringSeparator + InLineStarted);
+		FModelProtoStringParser::ArrayTToStringAuxiliary(String, TEXT("{0}FGraphProto.ValueInfo[{1}]: {2}"), InLineStarted, ValueInfo, FModelProtoStringParser::ToStringSeparator + InLineStarted);
+		FModelProtoStringParser::ArrayTToStringAuxiliary(String, TEXT("{0}FGraphProto.QuantizationAnnoation[{1}]: {2}"), InLineStarted, QuantizationAnnotation, FModelProtoStringParser::ToStringSeparator + InLineStarted);
+		return String;
+	}
+	else
+	{
+		return TEXT("");
+	}
+}
+
+
+
+/* FTrainingProto structors
+ *****************************************************************************/
+
+FTrainingInfoProto::FTrainingInfoProto()
+	: bIsLoaded(false)
+{
+}
+
+/* FTrainingInfoProto public functions
+ *****************************************************************************/ 
+
+bool FTrainingInfoProto::LoadFromString(const FString & InProtoString, const int32 InLevel)
+{
+	// Sanity check
+	if (InProtoString.Len() < 1)
+	{
+		return false;
+	}
+
+	// FString text into [key, value] maps
+	const TMap<FString, TArray<FString>> ProtoMap = FModelProtoStringParser::ModelProtoStringToMap(InProtoString, InLevel);
+
+	// Fill class variables
+	bIsLoaded = true;
+	bIsLoaded &= Initialization.LoadFromString(FModelProtoStringParser::GetModelProtoStringOrEmpty(ProtoMap, TEXT("initialization")), InLevel + 1);
+	bIsLoaded &= Algorithm.LoadFromString(FModelProtoStringParser::GetModelProtoStringOrEmpty(ProtoMap, TEXT("algorithm")), InLevel + 1);
+	bIsLoaded &= FModelProtoStringParser::GetModelProtoArray(InitializationBinding, ProtoMap, TEXT("initialization_binding"), InLevel + 1);
+	bIsLoaded &= FModelProtoStringParser::GetModelProtoArray(UpdateBinding, ProtoMap, TEXT("update_binding"), InLevel + 1);
+
+	if (!bIsLoaded)
+	{
+		UE_LOG(LogModelProto, Warning, TEXT("FTrainingInforProto::LoadFromString() failed."));
+	}
+	return bIsLoaded;
+}
+
+FString FTrainingInfoProto::ToString(const FString& InLineStarted) const
+{
+	if (bIsLoaded)
+	{
+		FString String = TEXT("\n");
+		String += FString::Format(TEXT("{0}FTrainingProto.Initialization: {1}"), { InLineStarted, Initialization.ToString(FModelProtoStringParser::ToStringSeparator + InLineStarted) });
+		String += FString::Format(TEXT("{0}FTrainingProto.Algorithm: {1}"), { InLineStarted, Algorithm.ToString(FModelProtoStringParser::ToStringSeparator + InLineStarted) });
+		FModelProtoStringParser::ArrayTToStringAuxiliary(String, TEXT("{0}FTrainingProto.InitializationBinding[{1}]: {2}"), InLineStarted, InitializationBinding, FModelProtoStringParser::ToStringSeparator + InLineStarted);
+		FModelProtoStringParser::ArrayTToStringAuxiliary(String, TEXT("{0}FTrainingProto.UpdateBinding[{1}]: {2}"), InLineStarted, UpdateBinding, FModelProtoStringParser::ToStringSeparator + InLineStarted);
+		return String;
+	}
+	else
+	{
+		return TEXT("");
+	}
+}
+
+
+
+/* FOperatorSetIdProto structors
+ *****************************************************************************/
+
+FOperatorSetIdProto::FOperatorSetIdProto()
+	: Version(-1)
+{
+}
+
+/* FOperatorSetIdProto public functions
+ *****************************************************************************/
+
+bool FOperatorSetIdProto::LoadFromString(const FString& InProtoString, const int32 InLevel)
+{
+	// Sanity check
+	if (InProtoString.Len() < 1)
+	{
+		return false;
+	}
+
+	// FString text into [key, value] maps
+	const TMap<FString, TArray<FString>> ProtoMap = FModelProtoStringParser::ModelProtoStringToMap(InProtoString, InLevel);
+
+	// Fill class variables
+	Domain = FModelProtoStringParser::GetModelProtoStringOrEmpty(ProtoMap, TEXT("domain"));
+	Version = FModelProtoStringParser::GetModelProtoInt64(ProtoMap, TEXT("version"));
+
+	const bool bIsLoaded = (Version > -1);
+	if (!bIsLoaded)
+	{
+		UE_LOG(LogModelProto, Warning, TEXT("FOperatorSetIdProto::LoadFromString() failed."));
+	}
+	return bIsLoaded;
+}
+
+FString FOperatorSetIdProto::ToString(const FString& InLineStarted) const
+{
+	FString String = TEXT("\n");
+	FModelProtoStringParser::StringToStringIfNotEmptyAuxiliary(String, TEXT("{0}FOperatorSetIdProto.Domain: {1}\n"), InLineStarted, Domain);
+	String += FString::Format(TEXT("{0}FOperatorSetIdProto.Version: {1}\n"), { InLineStarted, FString::FromInt(Version) });
+	return String;
+}
+
+
+
+/* FModelProto structors
+ *****************************************************************************/
+
+FModelProto::FModelProto()
+	: bIsLoaded(false)
+	, IRVersion(-1)
+	, ModelVersion(-1)
+{
+}
+
+/* FModelProto public functions
+ *****************************************************************************/
+
+bool FModelProto::LoadFromString(const FString& InProtoString, const int32 InLevel)
+{
+	// Sanity check
+	if (InProtoString.Len() < 1)
+	{
+		return false;
+	}
+
+	// FString text into [key, value] maps
+	const TMap<FString, TArray<FString>> ProtoMap = FModelProtoStringParser::ModelProtoStringToMap(InProtoString, InLevel);
+
+	// Fill class variables
+	bIsLoaded = true;
+	IRVersion = FModelProtoStringParser::GetModelProtoInt64(ProtoMap, TEXT("ir_version"));
+	bIsLoaded &= FModelProtoStringParser::GetModelProtoArray(OpsetImport, ProtoMap, TEXT("opset_import"), InLevel + 1);
+	ProducerName = FModelProtoStringParser::GetModelProtoStringOrEmpty(ProtoMap, TEXT("producer_name"));
+	ProducerVersion = FModelProtoStringParser::GetModelProtoStringOrEmpty(ProtoMap, TEXT("producer_version"));
+	Domain = FModelProtoStringParser::GetModelProtoStringOrEmpty(ProtoMap, TEXT("domain"));
+	ModelVersion = FModelProtoStringParser::GetModelProtoInt64(ProtoMap, TEXT("model_version"));
+	DocString = FModelProtoStringParser::GetModelProtoStringOrEmpty(ProtoMap, TEXT("doc_string"));
+	bIsLoaded &= Graph.LoadFromString(FModelProtoStringParser::GetModelProtoStringOrEmpty(ProtoMap, TEXT("graph")), InLevel + 1);
+	bIsLoaded &= FModelProtoStringParser::GetModelProtoArray(MetadataProps, ProtoMap, TEXT("metadata_props"), InLevel + 1);
+	bIsLoaded &= FModelProtoStringParser::GetModelProtoArray(TrainingInfo, ProtoMap, TEXT("training_info"), InLevel + 1);
+
+	bIsLoaded &= IRVersion > -1 && OpsetImport.Num() > 0;
+	if (!bIsLoaded)
+	{
+		UE_LOG(LogModelProto, Warning, TEXT("FModelProto::LoadFromString() failed."));
+	}
+	return bIsLoaded;
+}
+
+FString FModelProto::ToString(const FString& InLineStarted) const
+{
+	if (bIsLoaded)
+	{
+		FString String;
+		FModelProtoStringParser::PositiveInt32Or64ToStringAuxiliary(String, TEXT("{0}FModelProto.IRVersion: {1}\n"), InLineStarted, IRVersion);
+		FModelProtoStringParser::ArrayTToStringAuxiliary(String, TEXT("{0}FModelProto.OpsetImport[{1}]: {2}"), InLineStarted, OpsetImport, FModelProtoStringParser::ToStringSeparator + InLineStarted);
+		FModelProtoStringParser::StringToStringIfNotEmptyAuxiliary(String, TEXT("{0}FModelProto.ProducerName: {1}\n"), InLineStarted, ProducerName);
+		FModelProtoStringParser::StringToStringIfNotEmptyAuxiliary(String, TEXT("{0}FModelProto.ProducerVersion: {1}\n"), InLineStarted, ProducerVersion);
+		FModelProtoStringParser::StringToStringIfNotEmptyAuxiliary(String, TEXT("{0}FModelProto.Domain: {1}\n"), InLineStarted, Domain);
+		FModelProtoStringParser::PositiveInt32Or64ToStringAuxiliary(String, TEXT("{0}FModelProto.ModelVersion: {1}\n"), InLineStarted, ModelVersion);
+		FModelProtoStringParser::StringToStringIfNotEmptyAuxiliary(String, TEXT("{0}FModelProto.DocString: {1}\n"), InLineStarted, DocString);
+		String += FString::Format(TEXT("{0}FModelProto.Graph: {1}"), { InLineStarted, Graph.ToString(FModelProtoStringParser::ToStringSeparator + InLineStarted) });
+		FModelProtoStringParser::ArrayTToStringAuxiliary(String, TEXT("{0}FModelProto.MetadataProps[{1}]: {2}"), InLineStarted, MetadataProps, FModelProtoStringParser::ToStringSeparator + InLineStarted);
+		FModelProtoStringParser::ArrayTToStringAuxiliary(String, TEXT("{0}FModelProto.TrainingInfo[{1}]: {2}"), InLineStarted, TrainingInfo, FModelProtoStringParser::ToStringSeparator + InLineStarted);
+		return String;
+	}
+	else
+	{
+		return TEXT("");
+	}
+}
+
+const FGraphProto& FModelProto::GetGraph() const
+{
+	return Graph;
+}
diff --git a/Engine/Plugins/Experimental/NeuralNetworkInference/Source/Private/ModelProto/Private/ModelProtoModule.cpp b/Engine/Plugins/Experimental/NeuralNetworkInference/Source/Private/ModelProto/Private/ModelProtoModule.cpp
new file mode 100644
--- /dev/null
+++ b/Engine/Plugins/Experimental/NeuralNetworkInference/Source/Private/ModelProto/Private/ModelProtoModule.cpp
@@ -0,0 +1,13 @@
+// Copyright Epic Games, Inc. All Rights Reserved.
+
+#include "ModelProtoModule.h"
+
+void FModelProtoModule::StartupModule()
+{
+}
+
+void FModelProtoModule::ShutdownModule()
+{
+}
+
+IMPLEMENT_MODULE(FModelProtoModule, ModelProto);
diff --git a/Engine/Plugins/Experimental/NeuralNetworkInference/Source/Private/ModelProto/Private/ModelProtoModule.h b/Engine/Plugins/Experimental/NeuralNetworkInference/Source/Private/ModelProto/Private/ModelProtoModule.h
new file mode 100644
--- /dev/null
+++ b/Engine/Plugins/Experimental/NeuralNetworkInference/Source/Private/ModelProto/Private/ModelProtoModule.h
@@ -0,0 +1,17 @@
+// Copyright Epic Games, Inc. All Rights Reserved.
+
+#pragma once
+
+#include "CoreMinimal.h"
+#include "Modules/ModuleManager.h"
+
+/**
+* The public interface to this module
+*/
+class FModelProtoModule : public IModuleInterface
+{
+public:
+	// Begin IModuleInterface
+	virtual void StartupModule() override;
+	virtual void ShutdownModule() override;
+};
diff --git a/Engine/Plugins/Experimental/NeuralNetworkInference/Source/Private/ModelProto/Private/ModelProtoStringParser.cpp b/Engine/Plugins/Experimental/NeuralNetworkInference/Source/Private/ModelProto/Private/ModelProtoStringParser.cpp
new file mode 100644
--- /dev/null
+++ b/Engine/Plugins/Experimental/NeuralNetworkInference/Source/Private/ModelProto/Private/ModelProtoStringParser.cpp
@@ -0,0 +1,288 @@
+// Copyright Epic Games, Inc. All Rights Reserved.
+
+#include "ModelProtoStringParser.h"
+#include "ModelProtoUtils.h"
+
+
+
+/* FModelProtoStringParser static variables
+ *****************************************************************************/
+
+FString FModelProtoStringParser::LineSeparator = TEXT("  ");
+FString FModelProtoStringParser::ToStringSeparator = TEXT("    ");
+
+
+
+/* FModelProtoStringParser public functions
+ *****************************************************************************/
+
+TMap<FString, TArray<FString>> FModelProtoStringParser::ModelProtoStringToMap(const FString& InString, const int32 InLevel)
+{
+	TMap<FString, TArray<FString>> ProtoMap;
+
+	// Parse FString
+	FString RightSubString = (InString.Find(TEXT("\r\n")) == INDEX_NONE
+		// Windows uses "/r/n" for end-of-line, most others just "/n"
+		? InString.Replace(TEXT("\n"), TEXT("\r\n"))
+		: InString);
+	int32 EnterIndex;
+	while (RightSubString.FindChar(TCHAR('\n'), EnterIndex) || RightSubString.Len() > 0)
+	{
+		if (EnterIndex == INDEX_NONE)
+		{
+			EnterIndex = RightSubString.Len() + 1;
+		}
+		const FString SubString = RightSubString.Left(EnterIndex);
+		int32 SeparatorCharIndex;
+		FString Key;
+		FString Value;
+		// Colon
+		if (SubString.FindChar(TCHAR(':'), SeparatorCharIndex))
+		{
+			Key = SubString.Mid(LineSeparator.Len() * InLevel, SeparatorCharIndex - LineSeparator.Len() * InLevel);
+			Value = SubString.Mid(SeparatorCharIndex + 2, EnterIndex - SeparatorCharIndex - 3);
+			// Update RightSubString
+			RightSubString.RightChopInline(EnterIndex + 1);
+		}
+		// Braces
+		else if (SubString.FindChar(TCHAR('{'), SeparatorCharIndex))
+		{
+			// Key
+			Key = SubString.Mid(LineSeparator.Len() * InLevel, SeparatorCharIndex - LineSeparator.Len() * InLevel - 1);
+			// Get right ClosingBrace
+			FString ClosingBrace = TEXT("\r\n");
+			for (int32 Index = 0; Index < InLevel; ++Index)
+			{
+				ClosingBrace += LineSeparator;
+			}
+			ClosingBrace += TEXT("}");
+			// Value
+			const int32 ClosingBraceIndex = RightSubString.Find(ClosingBrace);
+			const int32 StartPoint = SeparatorCharIndex + 3;
+			if (ClosingBraceIndex != INDEX_NONE)
+			{
+				Value = RightSubString.Mid(StartPoint, ClosingBraceIndex - StartPoint);
+			}
+			else
+			{
+				UE_LOG(LogModelProto, Warning, TEXT("ModelProtoStringToMap(): Expected closing brace missing: \"}\"."));
+				return TMap<FString, TArray<FString>>();
+			}
+			// Update RightSubString
+			RightSubString.RightChopInline(ClosingBraceIndex + ClosingBrace.Len() + 2);
+		}
+		// Unknown
+		else
+		{
+			UE_LOG(LogModelProto, Warning, TEXT("ModelProtoStringToMap(): InLevel: %d\n - InString: %s\n - RightSubString: %s."),
+				InLevel, *InString, *RightSubString);
+			break;
+		}
+		// Already exist --> Emplace to TArray
+		if (TArray<FString>* Values = ProtoMap.Find(Key))
+		{
+			Values->Emplace(Value);
+		}
+		// First time --> Add
+		else
+		{
+			ProtoMap.Add(Key, TArray<FString>({Value}));
+		}
+	}
+
+	return ProtoMap;
+}
+
+FString FModelProtoStringParser::GetModelProtoStringOrEmpty(const TMap<FString, TArray<FString>>& InProtoMap, const FString& InKey)
+{
+	if (const TArray<FString>* const String = InProtoMap.Find(InKey))
+	{
+		if (String->Num() != 1)
+		{
+			UE_LOG(LogModelProto, Warning,
+				TEXT("GetModelProtoStringOrEmpty(): InKey: %s, InProtoMap.Num() = %d, String->Num() = %d != 1."),
+				*InKey, InProtoMap.Num(), String->Num());
+		}
+		FString OutputString = (*String)[0]; // Copy so RemoveQuotesFromProtoString does not edit it
+		// If string is something like: "gemm", it changes it into gemm
+		RemoveQuotesFromProtoString(OutputString);
+		return OutputString;
+	}
+	return TEXT("");
+}
+
+TArray<FString> FModelProtoStringParser::GetModelProtoStringArray(const TMap<FString, TArray<FString>>& InProtoMap, const FString& InKey)
+{
+	if (const TArray<FString>* const StringArray = InProtoMap.Find(InKey))
+	{
+		if (StringArray->Num() < 1)
+		{
+			UE_LOG(LogModelProto, Warning,
+				TEXT("GetModelProtoStringArray(): InKey: %s, InProtoMap.Num() = %d, StringArray->Num() = %d (should be >1)."),
+				*InKey, InProtoMap.Num(), StringArray->Num());
+		}
+		// If string is something like: "gemm", it changes it into gemm
+		TArray<FString> CleanedStringArray = *StringArray;
+		for (FString& String : CleanedStringArray)
+		{
+			RemoveQuotesFromProtoString(String);
+		}
+		return CleanedStringArray;
+	}
+	return TArray<FString>({});
+}
+
+int32 FModelProtoStringParser::GetModelProtoInt32(const TMap<FString, TArray<FString>>& InProtoMap, const FString& InKey)
+{
+	const FString ModelProtoNumberString = GetModelProtoStringOrEmpty(InProtoMap, InKey);
+	if (ModelProtoNumberString.Len() > 0)
+	{
+		return FCString::Atoi(*ModelProtoNumberString);
+	}
+	return int32(-1);
+}
+
+TArray<uint8> FModelProtoStringParser::GetModelProtoCharAsUInt8Array(const TMap<FString, TArray<FString>>& InProtoMap, const FString& InKey)
+{
+	TArray<uint8> ProtoCharAsUint8Array;
+	const FString RawDataString = FModelProtoStringParser::GetModelProtoStringOrEmpty(InProtoMap, InKey);
+	if (!RawDataString.IsEmpty())
+	{
+		const char* const RawDataCharPtr = TCHAR_TO_ANSI(*RawDataString);
+		// Allocate memory
+		ProtoCharAsUint8Array.SetNumUninitialized(RawDataString.Len());
+		// Copy memory
+		FMemory::Memcpy(ProtoCharAsUint8Array.GetData(), RawDataCharPtr, ProtoCharAsUint8Array.Num());
+	}
+	// Return array
+	return ProtoCharAsUint8Array;
+}
+
+TArray<int32> FModelProtoStringParser::GetModelProtoInt32Array(const TMap<FString, TArray<FString>>& InProtoMap, const FString& InKey)
+{
+	const TArray<FString> ModelProtoNumberStringArray = GetModelProtoStringArray(InProtoMap, InKey);
+	TArray<int32> ProtoNumberArray;
+	for (const FString& ModelProtoNumberString : ModelProtoNumberStringArray)
+	{
+		if (ModelProtoNumberString.Len() > 0)
+		{
+			ProtoNumberArray.Emplace(FCString::Atoi(*ModelProtoNumberString));
+		}
+	}
+	return ProtoNumberArray;
+}
+
+int64 FModelProtoStringParser::GetModelProtoInt64(const TMap<FString, TArray<FString>>& InProtoMap, const FString& InKey)
+{
+	const FString ModelProtoNumberString = GetModelProtoStringOrEmpty(InProtoMap, InKey);
+	if (ModelProtoNumberString.Len() > 0)
+	{
+		return FCString::Atoi64(*ModelProtoNumberString);
+	}
+	return int64(-1);
+}
+
+TArray<int64> FModelProtoStringParser::GetModelProtoInt64Array(const TMap<FString, TArray<FString>>& InProtoMap, const FString& InKey)
+{
+	const TArray<FString> ModelProtoNumberStringArray = GetModelProtoStringArray(InProtoMap, InKey);
+	TArray<int64> ProtoNumberArray;
+	for (const FString& ModelProtoNumberString : ModelProtoNumberStringArray)
+	{
+		if (ModelProtoNumberString.Len() > 0)
+		{
+			ProtoNumberArray.Emplace(FCString::Atoi64(*ModelProtoNumberString));
+		}
+	}
+	return ProtoNumberArray;
+}
+
+float FModelProtoStringParser::GetModelProtoFloat(const TMap<FString, TArray<FString>>& InProtoMap, const FString& InKey)
+{
+	const FString ModelProtoNumberString = GetModelProtoStringOrEmpty(InProtoMap, InKey);
+	if (ModelProtoNumberString.Len() > 0)
+	{
+		return FCString::Atof(*ModelProtoNumberString);
+	}
+	return int64(-1);
+}
+
+TArray<float> FModelProtoStringParser::GetModelProtoFloatArray(const TMap<FString, TArray<FString>>& InProtoMap, const FString& InKey)
+{
+	const TArray<FString> ModelProtoNumberStringArray = GetModelProtoStringArray(InProtoMap, InKey);
+	TArray<float> ProtoNumberArray;
+	for (const FString& ModelProtoNumberString : ModelProtoNumberStringArray)
+	{
+		if (ModelProtoNumberString.Len() > 0)
+		{
+			ProtoNumberArray.Emplace(FCString::Atof(*ModelProtoNumberString));
+		}
+	}
+	return ProtoNumberArray;
+}
+
+double FModelProtoStringParser::GetModelProtoDouble(const TMap<FString, TArray<FString>>& InProtoMap, const FString& InKey)
+{
+	const FString ModelProtoNumberString = GetModelProtoStringOrEmpty(InProtoMap, InKey);
+	if (ModelProtoNumberString.Len() > 0)
+	{
+		return FCString::Atod(*ModelProtoNumberString);
+	}
+	return int64(-1);
+}
+
+TArray<double> FModelProtoStringParser::GetModelProtoDoubleArray(const TMap<FString, TArray<FString>>& InProtoMap, const FString& InKey)
+{
+	const TArray<FString> ModelProtoNumberStringArray = GetModelProtoStringArray(InProtoMap, InKey);
+	TArray<double> ProtoNumberArray;
+	for (const FString& ModelProtoNumberString : ModelProtoNumberStringArray)
+	{
+		if (ModelProtoNumberString.Len() > 0)
+		{
+			ProtoNumberArray.Emplace(FCString::Atod(*ModelProtoNumberString));
+		}
+	}
+	return ProtoNumberArray;
+}
+
+TArray<uint64> FModelProtoStringParser::GetModelProtoUInt64Array(const TMap<FString, TArray<FString>>& InProtoMap, const FString& InKey)
+{
+	const TArray<FString> ModelProtoNumberStringArray = GetModelProtoStringArray(InProtoMap, InKey);
+	TArray<uint64> ProtoNumberArray;
+	for (const FString& ModelProtoNumberString : ModelProtoNumberStringArray)
+	{
+		if (ModelProtoNumberString.Len() > 0)
+		{
+			ProtoNumberArray.Emplace(FCString::Strtoui64(*ModelProtoNumberString, nullptr, 10));
+		}
+	}
+	return ProtoNumberArray;
+}
+
+void FModelProtoStringParser::StringArrayToStringAuxiliary(FString& InOutString, const TCHAR* const InText, const FString& InLineStarted, const TArray<FString>& InArrayString)
+{
+	if (InArrayString.Num() > 0)
+	{
+		FString FinalString = TEXT(" ");
+		for (const FString& String : InArrayString)
+		{
+			FinalString += String + TEXT(" ");
+		}
+		StringToStringIfNotEmptyAuxiliary(InOutString, InText, InLineStarted, FinalString);
+	}
+}
+
+void FModelProtoStringParser::StringToStringIfNotEmptyAuxiliary(FString& InOutString, const TCHAR* const InText, const FString& InLineStarted, const FString& InString)
+{
+	if (!InString.IsEmpty())
+	{
+		InOutString += FString::Format(InText, { InLineStarted, InString });
+	}
+}
+
+void FModelProtoStringParser::RemoveQuotesFromProtoString(FString& InOutProtoString)
+{
+	if (InOutProtoString.Len() > 1 && InOutProtoString.Left(1) == TEXT("\"") && InOutProtoString.Right(1) == TEXT("\""))
+	{
+		InOutProtoString.MidInline(1, InOutProtoString.Len() - 2);
+	}
+}
diff --git a/Engine/Plugins/Experimental/NeuralNetworkInference/Source/Private/ModelProto/Private/ModelProtoStringParser.h b/Engine/Plugins/Experimental/NeuralNetworkInference/Source/Private/ModelProto/Private/ModelProtoStringParser.h
new file mode 100644
--- /dev/null
+++ b/Engine/Plugins/Experimental/NeuralNetworkInference/Source/Private/ModelProto/Private/ModelProtoStringParser.h
@@ -0,0 +1,150 @@
+// Copyright Epic Games, Inc. All Rights Reserved.
+
+#pragma once
+
+#include "CoreMinimal.h"
+
+class FModelProtoStringParser
+{
+public:
+	static FString LineSeparator;
+
+	static FString ToStringSeparator;
+
+	// Note: InString is considerably big, avoiding the first InString copy with a move would be faster but would require
+	// InString not to be const
+	static TMap<FString, TArray<FString>> ModelProtoStringToMap(const FString& InString, const int32 InLevel);
+
+	static FString GetModelProtoStringOrEmpty(const TMap<FString, TArray<FString>>& InProtoMap, const FString& InKey);
+
+	static TArray<FString> GetModelProtoStringArray(const TMap<FString, TArray<FString>>& InProtoMap, const FString& InKey);
+
+	template<typename T>
+	static bool GetModelProtoArray(TArray<T>& OutArray, const TMap<FString, TArray<FString>>& InProtoMap, const FString& InKey, const int32 InLevel);
+
+	template<typename T>
+	static bool GetModelProtoMap(TMap<FString, T>& OutMap, const TMap<FString, TArray<FString>>& InProtoMap, const FString& InKey, const int32 InLevel);
+
+	static int32 GetModelProtoInt32(const TMap<FString, TArray<FString>>& InProtoMap, const FString& InKey);
+
+	static TArray<uint8> GetModelProtoCharAsUInt8Array(const TMap<FString, TArray<FString>>& InProtoMap, const FString& InKey);
+
+	static TArray<int32> GetModelProtoInt32Array(const TMap<FString, TArray<FString>>& InProtoMap, const FString& InKey);
+
+	static int64 GetModelProtoInt64(const TMap<FString, TArray<FString>>& InProtoMap, const FString& InKey);
+
+	static TArray<int64> GetModelProtoInt64Array(const TMap<FString, TArray<FString>>& InProtoMap, const FString& InKey);
+
+	static float GetModelProtoFloat(const TMap<FString, TArray<FString>>& InProtoMap, const FString& InKey);
+
+	static TArray<float> GetModelProtoFloatArray(const TMap<FString, TArray<FString>>& InProtoMap, const FString& InKey);
+
+	static double GetModelProtoDouble(const TMap<FString, TArray<FString>>& InProtoMap, const FString& InKey);
+
+	static TArray<double> GetModelProtoDoubleArray(const TMap<FString, TArray<FString>>& InProtoMap, const FString& InKey);
+
+	static TArray<uint64> GetModelProtoUInt64Array(const TMap<FString, TArray<FString>>& InProtoMap, const FString& inKey);
+
+	/**
+	 * Functions meant for ::ToString() debugging
+	 */
+	static void StringArrayToStringAuxiliary(FString& InOutString, const TCHAR* const InText, const FString& InLineStarted, const TArray<FString>& InArrayString);
+
+	template<typename T>
+	static void ArrayTToStringAuxiliary(FString& InOutString, const TCHAR* const InText, const FString& InLineStarted, const TArray<T>& InArray, const FString& InLineStartedNext);
+
+	template<typename T>
+	static void MapTToStringAuxiliary(FString& InOutString, const TCHAR* const InText, const FString& InLineStarted, const TMap<FString, T>& InMap, const FString& InLineStartedNext);
+
+	template<typename T>
+	static void PositiveInt32Or64ToStringAuxiliary(FString& InOutString, const TCHAR* const InText, const FString& InLineStarted, const T InNumberInt32Or64);
+
+	template<typename T>
+	static void GreaterThan0Int32Or64ToStringAuxiliary(FString& InOutString, const TCHAR* const InText, const FString& InLineStarted, const T InNumberInt32Or64);
+
+	static void StringToStringIfNotEmptyAuxiliary(FString& InOutString, const TCHAR* const InText, const FString& InLineStarted, const FString& InString);
+
+private:
+	static void RemoveQuotesFromProtoString(FString& InOutProtoString);
+};
+
+
+
+/* FModelProtoStringParser template functions
+ *****************************************************************************/
+
+template<typename T>
+bool FModelProtoStringParser::GetModelProtoArray(TArray<T>& OutArray, const TMap<FString, TArray<FString>>& InProtoMap, const FString& InKey, const int32 InLevel)
+{
+	OutArray.Empty();
+	const TArray<FString> StringArray = FModelProtoStringParser::GetModelProtoStringArray(InProtoMap, InKey);
+	for (const FString& String : StringArray)
+	{
+		OutArray.Emplace();
+		if (!OutArray.Last().LoadFromString(String, InLevel))
+		{
+			// If invalid --> Remove it and return false
+			OutArray.Pop();
+			return false;
+		}
+	}
+	return true;
+}
+
+template<typename T>
+bool FModelProtoStringParser::GetModelProtoMap(TMap<FString, T>& OutMap, const TMap<FString, TArray<FString>>& InProtoMap, const FString& InKey, const int32 InLevel)
+{
+	OutMap.Empty();
+	const TArray<FString> StringArray = FModelProtoStringParser::GetModelProtoStringArray(InProtoMap, InKey);
+	for (const FString& String : StringArray)
+	{
+		T TemporaryObject;
+		if (TemporaryObject.LoadFromString(String, InLevel))
+		{
+			OutMap.Add(TemporaryObject.Name, MoveTemp(TemporaryObject));
+		}
+		else
+		{
+			return false;
+		}
+	}
+	return true;
+}
+
+template<typename T>
+void FModelProtoStringParser::ArrayTToStringAuxiliary(FString& InOutString, const TCHAR* const InText, const FString& InLineStarted, const TArray<T>& InArray, const FString& InLineStartedNext)
+{
+	int32 Index = -1;
+	for (const T& Value : InArray)
+	{
+		InOutString += FString::Format(InText, { InLineStarted, ++Index, Value.ToString(InLineStartedNext) });
+	}
+}
+
+template<typename T>
+void FModelProtoStringParser::MapTToStringAuxiliary(FString& InOutString, const TCHAR* const InText, const FString& InLineStarted, const TMap<FString, T>& InMap, const FString& InLineStartedNext)
+{
+	int32 Index = -1;
+	for (const TTuple<FString, T>& Tuple : InMap)
+	{
+		InOutString += FString::Format(InText, { InLineStarted, ++Index, Tuple.Value.ToString(InLineStartedNext) });
+	}
+}
+
+template<typename T>
+void FModelProtoStringParser::PositiveInt32Or64ToStringAuxiliary(FString& InOutString, const TCHAR* const InText, const FString& InLineStarted, const T InNumberInt32Or64)
+{
+	if (InNumberInt32Or64 > -1)
+	{
+		InOutString += FString::Format(InText, { InLineStarted, FString::FromInt(InNumberInt32Or64) });
+	}
+}
+
+template<typename T>
+void FModelProtoStringParser::GreaterThan0Int32Or64ToStringAuxiliary(FString& InOutString, const TCHAR* const InText, const FString& InLineStarted, const T InNumberInt32Or64)
+{
+	if (InNumberInt32Or64 > 0)
+	{
+		InOutString += FString::Format(InText, { InLineStarted, FString::FromInt(InNumberInt32Or64) });
+	}
+}
diff --git a/Engine/Plugins/Experimental/NeuralNetworkInference/Source/Private/ModelProto/Private/ModelProtoUtils.h b/Engine/Plugins/Experimental/NeuralNetworkInference/Source/Private/ModelProto/Private/ModelProtoUtils.h
new file mode 100644
--- /dev/null
+++ b/Engine/Plugins/Experimental/NeuralNetworkInference/Source/Private/ModelProto/Private/ModelProtoUtils.h
@@ -0,0 +1,7 @@
+// Copyright Epic Games, Inc. All Rights Reserved.
+
+#pragma once
+
+#include "CoreMinimal.h"
+
+DEFINE_LOG_CATEGORY_STATIC(LogModelProto, Display, All);
diff --git a/Engine/Plugins/Experimental/NeuralNetworkInference/Source/Private/NeuralNetworkInferenceShaders/Internal/NeuralNetworkInferenceShaders/BatchNormalizationCS.h b/Engine/Plugins/Experimental/NeuralNetworkInference/Source/Private/NeuralNetworkInferenceShaders/Internal/NeuralNetworkInferenceShaders/BatchNormalizationCS.h
new file mode 100644
--- /dev/null
+++ b/Engine/Plugins/Experimental/NeuralNetworkInference/Source/Private/NeuralNetworkInferenceShaders/Internal/NeuralNetworkInferenceShaders/BatchNormalizationCS.h
@@ -0,0 +1,44 @@
+// Copyright Epic Games, Inc. All Rights Reserved.
+
+#pragma once
+
+#include "NeuralOperatorEnumClasses.h"
+// GPU/RHI/shaders
+#include "GlobalShader.h"
+#include "RHI.h"
+#include "ProfilingDebugging/RealtimeGPUProfiler.h"
+#include "RenderGraphUtils.h"
+#include "ShaderParameterUtils.h"
+
+
+
+class NEURALNETWORKINFERENCESHADERS_API FBatchNormalizationCS : public FGlobalShader
+{
+	DECLARE_SHADER_TYPE(FBatchNormalizationCS, Global);
+	SHADER_USE_PARAMETER_STRUCT(FBatchNormalizationCS, FGlobalShader)
+
+public:
+	class FBatchNormalizationMode : SHADER_PERMUTATION_ENUM_CLASS("BATCH_NORMALIZATION_MODE", EBatchNormalizationMode);
+	class FIsInlined : SHADER_PERMUTATION_BOOL("IS_INLINED");
+	using FPermutationDomain = TShaderPermutationDomain<FBatchNormalizationMode, FIsInlined>;
+
+	static const uint32 THREADGROUP_SIZE_X;
+
+	static void ModifyCompilationEnvironment(const FGlobalShaderPermutationParameters& InParameters, FShaderCompilerEnvironment& OutEnvironment);
+
+	BEGIN_SHADER_PARAMETER_STRUCT(FParameters, )
+		// Input variables
+		SHADER_PARAMETER(uint32, BatchSize)
+		SHADER_PARAMETER(uint32, ChannelSize)
+		SHADER_PARAMETER(uint32, ChannelsVolume)
+		SHADER_PARAMETER(uint32, ImageArea)
+		SHADER_PARAMETER(float, Epsilon)
+		// SRV/UAV variables
+		SHADER_PARAMETER_RDG_BUFFER_SRV(Buffer<float>, XSRV)
+		SHADER_PARAMETER_RDG_BUFFER_SRV(Buffer<float>, Scale)
+		SHADER_PARAMETER_RDG_BUFFER_SRV(Buffer<float>, Bias)
+		SHADER_PARAMETER_RDG_BUFFER_SRV(Buffer<float>, Mean)
+		SHADER_PARAMETER_RDG_BUFFER_SRV(Buffer<float>, Variance)
+		SHADER_PARAMETER_RDG_BUFFER_UAV(RWBuffer<float>, OutputUAV)
+	END_SHADER_PARAMETER_STRUCT()
+};
diff --git a/Engine/Plugins/Experimental/NeuralNetworkInference/Source/Private/NeuralNetworkInferenceShaders/Internal/NeuralNetworkInferenceShaders/ConvBaseCS.h b/Engine/Plugins/Experimental/NeuralNetworkInference/Source/Private/NeuralNetworkInferenceShaders/Internal/NeuralNetworkInferenceShaders/ConvBaseCS.h
new file mode 100644
--- /dev/null
+++ b/Engine/Plugins/Experimental/NeuralNetworkInference/Source/Private/NeuralNetworkInferenceShaders/Internal/NeuralNetworkInferenceShaders/ConvBaseCS.h
@@ -0,0 +1,58 @@
+// Copyright Epic Games, Inc. All Rights Reserved.
+
+#pragma once
+
+#include "NeuralOperatorEnumClasses.h"
+// GPU/RHI/shaders
+#include "GlobalShader.h"
+#include "RHI.h"
+#include "ProfilingDebugging/RealtimeGPUProfiler.h"
+#include "RenderGraphUtils.h"
+#include "ShaderParameterUtils.h"
+
+
+
+class NEURALNETWORKINFERENCESHADERS_API FConvBaseCS : public FGlobalShader
+{
+	DECLARE_SHADER_TYPE(FConvBaseCS, Global);
+	SHADER_USE_PARAMETER_STRUCT(FConvBaseCS, FGlobalShader)
+
+public:
+	class FHasBias : SHADER_PERMUTATION_BOOL("HAS_BIAS");
+	class FConvMode : SHADER_PERMUTATION_ENUM_CLASS("CONV_MODE", EConvMode);
+	using FPermutationDomain = TShaderPermutationDomain<FHasBias, FConvMode>;
+
+	static const uint32 THREADGROUP_SIZE_X;
+	static const uint32 MAX_NUMBER_DIMENSIONS;
+
+	static void ModifyCompilationEnvironment(const FGlobalShaderPermutationParameters& InParameters, FShaderCompilerEnvironment& OutEnvironment);
+
+	BEGIN_SHADER_PARAMETER_STRUCT(FParameters, )
+		// All variables are used in all (1D, 2D, 3D, nD) convolutions
+		// Input variables
+		SHADER_PARAMETER(uint32, bIsTransposed)
+		SHADER_PARAMETER(uint32, Group)
+		SHADER_PARAMETER(uint32, MIntoGroup)
+		SHADER_PARAMETER(uint32, CIntoGroup)
+		SHADER_PARAMETER(uint32, OutputVolume)
+		SHADER_PARAMETER(uint32, OutputBatchVolume)
+		SHADER_PARAMETER(uint32, OutputImageArea)
+		SHADER_PARAMETER(uint32, WBatchVolume)
+		SHADER_PARAMETER(uint32, WImageArea)
+		SHADER_PARAMETER(uint32, XOrXWithZerosBatchVolume)
+		SHADER_PARAMETER(uint32, XOrXWithZerosImageArea)
+		SHADER_PARAMETER(int32, NumberConvolutionalDimensions) // Optional - Only for nD convolution
+		// Input SRV variables
+		SHADER_PARAMETER_SRV(Buffer<uint>, OutputSizes)
+		SHADER_PARAMETER_SRV(Buffer<uint>, XOrXWithZerosSizes)
+		SHADER_PARAMETER_SRV(Buffer<uint>, WSizes)
+		SHADER_PARAMETER_SRV(Buffer<uint>, Dilations)
+		SHADER_PARAMETER_SRV(Buffer<uint>, Strides)
+		SHADER_PARAMETER_SRV(Buffer<uint>, Pads)
+		// SRV/UAV variables
+		SHADER_PARAMETER_RDG_BUFFER_SRV(Buffer<float>, XOrXWithZerosSRV)
+		SHADER_PARAMETER_RDG_BUFFER_SRV(Buffer<float>, WSRV)
+		SHADER_PARAMETER_RDG_BUFFER_SRV(Buffer<float>, BSRV) // Optional - Only if bias
+		SHADER_PARAMETER_RDG_BUFFER_UAV(RWBuffer<float>, OutputUAV)
+	END_SHADER_PARAMETER_STRUCT()
+};
diff --git a/Engine/Plugins/Experimental/NeuralNetworkInference/Source/Private/NeuralNetworkInferenceShaders/Internal/NeuralNetworkInferenceShaders/ConvTransposeCS.h b/Engine/Plugins/Experimental/NeuralNetworkInference/Source/Private/NeuralNetworkInferenceShaders/Internal/NeuralNetworkInferenceShaders/ConvTransposeCS.h
new file mode 100644
--- /dev/null
+++ b/Engine/Plugins/Experimental/NeuralNetworkInference/Source/Private/NeuralNetworkInferenceShaders/Internal/NeuralNetworkInferenceShaders/ConvTransposeCS.h
@@ -0,0 +1,41 @@
+// Copyright Epic Games, Inc. All Rights Reserved.
+
+#pragma once
+
+#include "NeuralOperatorEnumClasses.h"
+// GPU/RHI/shaders
+#include "GlobalShader.h"
+#include "RHI.h"
+#include "ProfilingDebugging/RealtimeGPUProfiler.h"
+#include "RenderGraphUtils.h"
+#include "ShaderParameterUtils.h"
+
+
+
+class NEURALNETWORKINFERENCESHADERS_API FConvTransposeCS : public FGlobalShader
+{
+	DECLARE_SHADER_TYPE(FConvTransposeCS, Global);
+	SHADER_USE_PARAMETER_STRUCT(FConvTransposeCS, FGlobalShader)
+
+public:
+	class FConvMode : SHADER_PERMUTATION_ENUM_CLASS("CONV_MODE", EConvMode);
+	using FPermutationDomain = TShaderPermutationDomain<FConvMode>;
+
+	static const uint32 THREADGROUP_SIZE_X;
+
+	static void ModifyCompilationEnvironment(const FGlobalShaderPermutationParameters& InParameters, FShaderCompilerEnvironment& OutEnvironment);
+
+	BEGIN_SHADER_PARAMETER_STRUCT(FParameters, )
+		// All variables are used in all (1D, 2D, 3D, nD) convolutions
+		// Input variables
+		SHADER_PARAMETER(uint32, XVolume)
+		SHADER_PARAMETER(int32, NumberConvolutionalDimensions) // Optional - Only for nD convolution
+		// Input SRV variables
+		SHADER_PARAMETER_SRV(Buffer<uint>, Zeros)
+		SHADER_PARAMETER_SRV(Buffer<uint>, XSizes)
+		SHADER_PARAMETER_SRV(Buffer<uint>, XWithZerosSizes)
+		// SRV/UAV variables
+		SHADER_PARAMETER_RDG_BUFFER_SRV(Buffer<float>, XSRV)
+		SHADER_PARAMETER_RDG_BUFFER_UAV(RWBuffer<float>, XWithZerosUAV)
+	END_SHADER_PARAMETER_STRUCT()
+};
diff --git a/Engine/Plugins/Experimental/NeuralNetworkInference/Source/Private/NeuralNetworkInferenceShaders/Internal/NeuralNetworkInferenceShaders/CopyCS.h b/Engine/Plugins/Experimental/NeuralNetworkInference/Source/Private/NeuralNetworkInferenceShaders/Internal/NeuralNetworkInferenceShaders/CopyCS.h
new file mode 100644
--- /dev/null
+++ b/Engine/Plugins/Experimental/NeuralNetworkInference/Source/Private/NeuralNetworkInferenceShaders/Internal/NeuralNetworkInferenceShaders/CopyCS.h
@@ -0,0 +1,32 @@
+// Copyright Epic Games, Inc. All Rights Reserved.
+
+#pragma once
+
+#include "NeuralOperatorEnumClasses.h"
+// GPU/RHI/shaders
+#include "GlobalShader.h"
+#include "RHI.h"
+#include "ProfilingDebugging/RealtimeGPUProfiler.h"
+#include "RenderGraphUtils.h"
+#include "ShaderParameterUtils.h"
+
+
+
+class NEURALNETWORKINFERENCESHADERS_API FCopyCS : public FGlobalShader
+{
+	DECLARE_SHADER_TYPE(FCopyCS, Global);
+	SHADER_USE_PARAMETER_STRUCT(FCopyCS, FGlobalShader)
+
+public:
+	static const uint32 THREADGROUP_SIZE_X;
+
+	static void ModifyCompilationEnvironment(const FGlobalShaderPermutationParameters& InParameters, FShaderCompilerEnvironment& OutEnvironment);
+
+	BEGIN_SHADER_PARAMETER_STRUCT(FParameters, )
+		// Input variables
+		SHADER_PARAMETER(uint32, Num)
+		// SRV/UAV variables
+		SHADER_PARAMETER_RDG_BUFFER_SRV(Buffer<float>, InputSRV)
+		SHADER_PARAMETER_RDG_BUFFER_UAV(RWBuffer<float>, OutputUAV)
+	END_SHADER_PARAMETER_STRUCT()
+};
diff --git a/Engine/Plugins/Experimental/NeuralNetworkInference/Source/Private/NeuralNetworkInferenceShaders/Internal/NeuralNetworkInferenceShaders/ElementWiseCS.h b/Engine/Plugins/Experimental/NeuralNetworkInference/Source/Private/NeuralNetworkInferenceShaders/Internal/NeuralNetworkInferenceShaders/ElementWiseCS.h
new file mode 100644
--- /dev/null
+++ b/Engine/Plugins/Experimental/NeuralNetworkInference/Source/Private/NeuralNetworkInferenceShaders/Internal/NeuralNetworkInferenceShaders/ElementWiseCS.h
@@ -0,0 +1,43 @@
+// Copyright Epic Games, Inc. All Rights Reserved.
+
+#pragma once
+
+#include "NeuralOperatorEnumClasses.h"
+// GPU/RHI/shaders
+#include "GlobalShader.h"
+#include "RHI.h"
+#include "ProfilingDebugging/RealtimeGPUProfiler.h"
+#include "RenderGraphUtils.h"
+#include "ShaderParameterUtils.h"
+
+
+
+class NEURALNETWORKINFERENCESHADERS_API FElementWiseCS : public FGlobalShader
+{
+	DECLARE_SHADER_TYPE(FElementWiseCS, Global);
+	SHADER_USE_PARAMETER_STRUCT(FElementWiseCS, FGlobalShader)
+
+public:
+	static const uint32 THREADGROUP_SIZE_X;
+
+	class FShaderType : SHADER_PERMUTATION_ENUM_CLASS("SHADER_FUNCTION", EElementWiseOperator);
+	class FIsInlined : SHADER_PERMUTATION_BOOL("IS_INLINED");
+	using FPermutationDomain = TShaderPermutationDomain<FShaderType, FIsInlined>;
+
+	/** Set desired multidirectional operator. */
+	static void ModifyCompilationEnvironment(const FGlobalShaderPermutationParameters& InParameters, FShaderCompilerEnvironment& OutEnvironment);
+
+	BEGIN_SHADER_PARAMETER_STRUCT(FParameters, )
+		// Input variables
+		SHADER_PARAMETER(uint32, TensorSize)
+		SHADER_PARAMETER(float, Attribute)
+		// SRV/UAV variables
+		SHADER_PARAMETER_RDG_BUFFER_UAV(RWBuffer<float>, OutputUAV)
+		// Optional SRV/UAV variables
+		SHADER_PARAMETER_RDG_BUFFER_SRV(Buffer<float>, InputSRV)
+	END_SHADER_PARAMETER_STRUCT()
+
+private:
+	static FString GetFunctionString(const EElementWiseOperator InType);
+	static uint32 GetNumberAttributes(const EElementWiseOperator InType);
+};
diff --git a/Engine/Plugins/Experimental/NeuralNetworkInference/Source/Private/NeuralNetworkInferenceShaders/Internal/NeuralNetworkInferenceShaders/GemmCS.h b/Engine/Plugins/Experimental/NeuralNetworkInference/Source/Private/NeuralNetworkInferenceShaders/Internal/NeuralNetworkInferenceShaders/GemmCS.h
new file mode 100644
--- /dev/null
+++ b/Engine/Plugins/Experimental/NeuralNetworkInference/Source/Private/NeuralNetworkInferenceShaders/Internal/NeuralNetworkInferenceShaders/GemmCS.h
@@ -0,0 +1,50 @@
+// Copyright Epic Games, Inc. All Rights Reserved.
+
+#pragma once
+
+#include "NeuralOperatorEnumClasses.h"
+// GPU/RHI/shaders
+#include "GlobalShader.h"
+#include "RHI.h"
+#include "ProfilingDebugging/RealtimeGPUProfiler.h"
+#include "RenderGraphUtils.h"
+#include "ShaderParameterUtils.h"
+
+
+
+class NEURALNETWORKINFERENCESHADERS_API FGemmCS : public FGlobalShader
+{
+	DECLARE_SHADER_TYPE(FGemmCS, Global);
+	SHADER_USE_PARAMETER_STRUCT(FGemmCS, FGlobalShader)
+
+public:
+	class FGemmMode : SHADER_PERMUTATION_ENUM_CLASS("GEMM_MODE", EGemmMode);
+	using FPermutationDomain = TShaderPermutationDomain<FGemmMode>;
+
+	static const uint32 THREADGROUP_SIZE_X;
+	static const uint32 THREADGROUP_SIZE_Y;
+
+	static void ModifyCompilationEnvironment(const FGlobalShaderPermutationParameters& InParameters, FShaderCompilerEnvironment& OutEnvironment);
+
+	BEGIN_SHADER_PARAMETER_STRUCT(FParameters, )
+		// Input variables
+		SHADER_PARAMETER(float, Alpha)
+		SHADER_PARAMETER(float, Beta)
+		SHADER_PARAMETER(uint32, OutputRows)
+		SHADER_PARAMETER(uint32, OutputColumns)
+		SHADER_PARAMETER(uint32, AColsOrBRows)
+		SHADER_PARAMETER(uint32, AStrideX)
+		SHADER_PARAMETER(uint32, AStrideY)
+		SHADER_PARAMETER(uint32, BStrideX)
+		SHADER_PARAMETER(uint32, BStrideY)
+		SHADER_PARAMETER(uint32, CSizeX)
+		SHADER_PARAMETER(uint32, CSizeY)
+		SHADER_PARAMETER(uint32, OutputStride)
+		// SRV/UAV variables
+		SHADER_PARAMETER_RDG_BUFFER_SRV(Buffer<float>, ASRV)
+		SHADER_PARAMETER_RDG_BUFFER_SRV(Buffer<float>, BSRV)
+		SHADER_PARAMETER_RDG_BUFFER_SRV(Buffer<float>, CSRV)
+		SHADER_PARAMETER(float, BetaTimesCScalar)
+		SHADER_PARAMETER_RDG_BUFFER_UAV(RWBuffer<float>, OutputUAV)
+	END_SHADER_PARAMETER_STRUCT()
+};
diff --git a/Engine/Plugins/Experimental/NeuralNetworkInference/Source/Private/NeuralNetworkInferenceShaders/Internal/NeuralNetworkInferenceShaders/MultidirectionalBroadcastCS.h b/Engine/Plugins/Experimental/NeuralNetworkInference/Source/Private/NeuralNetworkInferenceShaders/Internal/NeuralNetworkInferenceShaders/MultidirectionalBroadcastCS.h
new file mode 100644
--- /dev/null
+++ b/Engine/Plugins/Experimental/NeuralNetworkInference/Source/Private/NeuralNetworkInferenceShaders/Internal/NeuralNetworkInferenceShaders/MultidirectionalBroadcastCS.h
@@ -0,0 +1,48 @@
+// Copyright Epic Games, Inc. All Rights Reserved.
+
+#pragma once
+
+#include "NeuralOperatorEnumClasses.h"
+// GPU/RHI/shaders
+#include "GlobalShader.h"
+#include "RHI.h"
+#include "ProfilingDebugging/RealtimeGPUProfiler.h"
+#include "RenderGraphUtils.h"
+#include "ShaderParameterUtils.h"
+
+
+
+class NEURALNETWORKINFERENCESHADERS_API FMultidirectionalBroadcastCS : public FGlobalShader
+{
+	DECLARE_SHADER_TYPE(FMultidirectionalBroadcastCS, Global);
+	SHADER_USE_PARAMETER_STRUCT(FMultidirectionalBroadcastCS, FGlobalShader)
+
+public:
+	static const uint32 THREADGROUP_SIZE_X;
+	static const uint32 MAX_NUMBER_DIMENSIONS;
+
+	class FShaderType : SHADER_PERMUTATION_ENUM_CLASS("SHADER_FUNCTION", EMultidirectionalBroadcastOperator);
+	class FInlinedMode : SHADER_PERMUTATION_ENUM_CLASS("INLINED_MODE", EMultidirectionalBroadcastInlinedMode);
+	class FShapeMode : SHADER_PERMUTATION_ENUM_CLASS("SHAPE_MODE", EMultidirectionalBroadcastShapeMode);
+	using FPermutationDomain = TShaderPermutationDomain<FShaderType, FInlinedMode, FShapeMode>;
+
+	static void ModifyCompilationEnvironment(const FGlobalShaderPermutationParameters& InParameters, FShaderCompilerEnvironment& OutEnvironment);
+
+	BEGIN_SHADER_PARAMETER_STRUCT(FParameters, )
+		// Input variables
+		SHADER_PARAMETER(uint32, TensorSize)
+		// SRV/UAV variables
+		SHADER_PARAMETER_RDG_BUFFER_UAV(RWBuffer<float>, OutputUAV)
+		// Optional elementwise SRV variables
+		SHADER_PARAMETER_RDG_BUFFER_SRV(Buffer<float>, ASRV)
+		SHADER_PARAMETER_RDG_BUFFER_SRV(Buffer<float>, BSRV)
+		// Optional multidirectional broadcast variables
+		SHADER_PARAMETER(uint32, ShapeDimensions)
+		SHADER_PARAMETER_SRV(Buffer<float>, ShapeOutput)
+		SHADER_PARAMETER_SRV(Buffer<float>, ShapeA)
+		SHADER_PARAMETER_SRV(Buffer<float>, ShapeB)
+	END_SHADER_PARAMETER_STRUCT()
+
+private:
+	static FString GetFunctionString(const EMultidirectionalBroadcastOperator InType);
+};
diff --git a/Engine/Plugins/Experimental/NeuralNetworkInference/Source/Private/NeuralNetworkInferenceShaders/NeuralNetworkInferenceShaders.Build.cs b/Engine/Plugins/Experimental/NeuralNetworkInference/Source/Private/NeuralNetworkInferenceShaders/NeuralNetworkInferenceShaders.Build.cs
new file mode 100644
--- /dev/null
+++ b/Engine/Plugins/Experimental/NeuralNetworkInference/Source/Private/NeuralNetworkInferenceShaders/NeuralNetworkInferenceShaders.Build.cs
@@ -0,0 +1,37 @@
+// Copyright Epic Games, Inc. All Rights Reserved.
+
+using UnrealBuildTool;
+
+public class NeuralNetworkInferenceShaders : ModuleRules
+{
+	public NeuralNetworkInferenceShaders( ReadOnlyTargetRules Target ) : base( Target )
+	{
+        ShortName = "NNIShaders"; // Shorten to avoid path-too-long errors
+		PCHUsage = ModuleRules.PCHUsageMode.UseExplicitOrSharedPCHs;
+
+		PublicIncludePaths.AddRange(
+			new string[] {
+				System.IO.Path.Combine(ModuleDirectory, "..")
+			}
+		);
+
+		PublicDependencyModuleNames.AddRange
+			(
+			new string[] {
+				"Core",
+				"CoreUObject",
+				"Engine",
+				"NeuralNetworkInferenceCore"
+			}
+		);
+
+		PrivateDependencyModuleNames.AddRange
+			(
+			new string[] {
+				"Projects",
+				"RenderCore",
+				"RHI"
+			}
+		);
+	}
+}
diff --git a/Engine/Plugins/Experimental/NeuralNetworkInference/Source/Private/NeuralNetworkInferenceShaders/Private/NeuralNetworkInferenceShaders/BatchNormalizationCS.cpp b/Engine/Plugins/Experimental/NeuralNetworkInference/Source/Private/NeuralNetworkInferenceShaders/Private/NeuralNetworkInferenceShaders/BatchNormalizationCS.cpp
new file mode 100644
--- /dev/null
+++ b/Engine/Plugins/Experimental/NeuralNetworkInference/Source/Private/NeuralNetworkInferenceShaders/Private/NeuralNetworkInferenceShaders/BatchNormalizationCS.cpp
@@ -0,0 +1,25 @@
+// Copyright Epic Games, Inc. All Rights Reserved.
+
+#include "NeuralNetworkInferenceShaders/BatchNormalizationCS.h"
+
+
+
+/* FBatchNormalizationCS public functions
+ *****************************************************************************/
+
+void FBatchNormalizationCS::ModifyCompilationEnvironment(const FGlobalShaderPermutationParameters& InParameters, FShaderCompilerEnvironment& OutEnvironment)
+{
+	FGlobalShader::ModifyCompilationEnvironment(InParameters, OutEnvironment);
+	OutEnvironment.SetDefine(TEXT("THREADGROUP_SIZE_X"), THREADGROUP_SIZE_X);
+}
+
+
+
+const uint32 FBatchNormalizationCS::THREADGROUP_SIZE_X = 128;
+
+
+
+/* Shader implementation
+ *****************************************************************************/
+
+IMPLEMENT_SHADER_TYPE(, FBatchNormalizationCS, TEXT("/Plugins/NeuralNetworkInference/Private/BatchNormalizationOperator.usf"), TEXT("BatchNormalizationCS"), SF_Compute) // Path defined in NeuralNetworkInferenceShadersModule.cpp
diff --git a/Engine/Plugins/Experimental/NeuralNetworkInference/Source/Private/NeuralNetworkInferenceShaders/Private/NeuralNetworkInferenceShaders/ConvBaseCS.cpp b/Engine/Plugins/Experimental/NeuralNetworkInference/Source/Private/NeuralNetworkInferenceShaders/Private/NeuralNetworkInferenceShaders/ConvBaseCS.cpp
new file mode 100644
--- /dev/null
+++ b/Engine/Plugins/Experimental/NeuralNetworkInference/Source/Private/NeuralNetworkInferenceShaders/Private/NeuralNetworkInferenceShaders/ConvBaseCS.cpp
@@ -0,0 +1,30 @@
+// Copyright Epic Games, Inc. All Rights Reserved.
+
+#include "NeuralNetworkInferenceShaders/ConvBaseCS.h"
+
+
+
+/* FConvBaseCS static members
+ *****************************************************************************/
+
+const uint32 FConvBaseCS::THREADGROUP_SIZE_X(128);
+const uint32 FConvBaseCS::MAX_NUMBER_DIMENSIONS(16);
+
+
+
+/* FConvBaseCS public functions
+ *****************************************************************************/
+
+void FConvBaseCS::ModifyCompilationEnvironment(const FGlobalShaderPermutationParameters& InParameters, FShaderCompilerEnvironment& OutEnvironment)
+{
+	FGlobalShader::ModifyCompilationEnvironment(InParameters, OutEnvironment);
+	OutEnvironment.SetDefine(TEXT("THREADGROUP_SIZE_X"), THREADGROUP_SIZE_X);
+	OutEnvironment.SetDefine(TEXT("MAX_NUMBER_DIMENSIONS"), MAX_NUMBER_DIMENSIONS);
+}
+
+
+
+/* Shader implementation
+ *****************************************************************************/
+
+IMPLEMENT_SHADER_TYPE(, FConvBaseCS, TEXT("/Plugins/NeuralNetworkInference/Private/ConvBaseOperator.usf"), TEXT("ConvBaseCS"), SF_Compute) // Path defined in NeuralNetworkInferenceShadersModule.cpp
diff --git a/Engine/Plugins/Experimental/NeuralNetworkInference/Source/Private/NeuralNetworkInferenceShaders/Private/NeuralNetworkInferenceShaders/ConvTransposeCS.cpp b/Engine/Plugins/Experimental/NeuralNetworkInference/Source/Private/NeuralNetworkInferenceShaders/Private/NeuralNetworkInferenceShaders/ConvTransposeCS.cpp
new file mode 100644
--- /dev/null
+++ b/Engine/Plugins/Experimental/NeuralNetworkInference/Source/Private/NeuralNetworkInferenceShaders/Private/NeuralNetworkInferenceShaders/ConvTransposeCS.cpp
@@ -0,0 +1,28 @@
+// Copyright Epic Games, Inc. All Rights Reserved.
+
+#include "NeuralNetworkInferenceShaders/ConvTransposeCS.h"
+
+
+
+/* FConvTransposeCS static members
+ *****************************************************************************/
+
+const uint32 FConvTransposeCS::THREADGROUP_SIZE_X(128);
+
+
+
+/* FConvTransposeCS public functions
+ *****************************************************************************/
+
+void FConvTransposeCS::ModifyCompilationEnvironment(const FGlobalShaderPermutationParameters& InParameters, FShaderCompilerEnvironment& OutEnvironment)
+{
+	FGlobalShader::ModifyCompilationEnvironment(InParameters, OutEnvironment);
+	OutEnvironment.SetDefine(TEXT("THREADGROUP_SIZE_X"), THREADGROUP_SIZE_X);
+}
+
+
+
+/* Shader implementation
+ *****************************************************************************/
+
+IMPLEMENT_SHADER_TYPE(, FConvTransposeCS, TEXT("/Plugins/NeuralNetworkInference/Private/ConvTransposeOperator.usf"), TEXT("XToXWithZerosCS"), SF_Compute) // Path defined in NeuralNetworkInferenceShadersModule.cpp
diff --git a/Engine/Plugins/Experimental/NeuralNetworkInference/Source/Private/NeuralNetworkInferenceShaders/Private/NeuralNetworkInferenceShaders/CopyCS.cpp b/Engine/Plugins/Experimental/NeuralNetworkInference/Source/Private/NeuralNetworkInferenceShaders/Private/NeuralNetworkInferenceShaders/CopyCS.cpp
new file mode 100644
--- /dev/null
+++ b/Engine/Plugins/Experimental/NeuralNetworkInference/Source/Private/NeuralNetworkInferenceShaders/Private/NeuralNetworkInferenceShaders/CopyCS.cpp
@@ -0,0 +1,25 @@
+// Copyright Epic Games, Inc. All Rights Reserved.
+
+#include "NeuralNetworkInferenceShaders/CopyCS.h"
+
+
+
+/* FCopyCS public functions
+ *****************************************************************************/
+
+void FCopyCS::ModifyCompilationEnvironment(const FGlobalShaderPermutationParameters& InParameters, FShaderCompilerEnvironment& OutEnvironment)
+{
+	FGlobalShader::ModifyCompilationEnvironment(InParameters, OutEnvironment);
+	OutEnvironment.SetDefine(TEXT("THREADGROUP_SIZE_X"), THREADGROUP_SIZE_X);
+}
+
+
+
+const uint32 FCopyCS::THREADGROUP_SIZE_X = 128;
+
+
+
+/* Shader implementation
+ *****************************************************************************/
+
+IMPLEMENT_SHADER_TYPE(, FCopyCS, TEXT("/Plugins/NeuralNetworkInference/Private/Copy.usf"), TEXT("CopyCS"), SF_Compute) // Path defined in NeuralNetworkInferenceShadersModule.cpp
diff --git a/Engine/Plugins/Experimental/NeuralNetworkInference/Source/Private/NeuralNetworkInferenceShaders/Private/NeuralNetworkInferenceShaders/ElementWiseCS.cpp b/Engine/Plugins/Experimental/NeuralNetworkInference/Source/Private/NeuralNetworkInferenceShaders/Private/NeuralNetworkInferenceShaders/ElementWiseCS.cpp
new file mode 100644
--- /dev/null
+++ b/Engine/Plugins/Experimental/NeuralNetworkInference/Source/Private/NeuralNetworkInferenceShaders/Private/NeuralNetworkInferenceShaders/ElementWiseCS.cpp
@@ -0,0 +1,115 @@
+// Copyright Epic Games, Inc. All Rights Reserved.
+
+#include "NeuralNetworkInferenceShaders/ElementWiseCS.h"
+#include "NeuralNetworkInferenceShaders/Utils.h"
+
+
+
+/* FElementWiseCS static members
+ *****************************************************************************/
+
+const uint32 FElementWiseCS::THREADGROUP_SIZE_X(128);
+
+
+
+/* FElementWiseCS public functions
+ *****************************************************************************/
+
+void FElementWiseCS::ModifyCompilationEnvironment(const FGlobalShaderPermutationParameters& InParameters, FShaderCompilerEnvironment& OutEnvironment)
+{
+	FGlobalShader::ModifyCompilationEnvironment(InParameters, OutEnvironment);
+	OutEnvironment.SetDefine(TEXT("THREADGROUP_SIZE_X"), THREADGROUP_SIZE_X);
+	// NumberAttributes
+	FPermutationDomain PermutationVector(InParameters.PermutationId);
+	const uint32 NumberAttributes = GetNumberAttributes(PermutationVector.Get<FShaderType>());
+	OutEnvironment.SetDefine(TEXT("NUMBER_ATTRIBUTES"), NumberAttributes);
+	// ElementWiseFunction
+	const FString ElementWiseFunction = GetFunctionString(PermutationVector.Get<FShaderType>());
+	if (NumberAttributes == 0)
+	{
+		OutEnvironment.SetDefine(TEXT("ElementWiseFunction(X)"), *ElementWiseFunction);
+	}
+	else
+	{
+		OutEnvironment.SetDefine(TEXT("ElementWiseFunction(X, Attribute)"), *ElementWiseFunction);
+	}
+	// Note: IS_INLINED is already defined because we used SHADER_PERMUTATION_BOOL("IS_INLINED")
+}
+
+
+
+/* FElementWiseCS private static functions
+ *****************************************************************************/
+
+FString FElementWiseCS::GetFunctionString(const EElementWiseOperator InType)
+{
+	if (InType >= EElementWiseOperator::MAX)
+	{
+		UE_LOG(LogNeuralNetworkInferenceShaders, Warning, TEXT("FElementWiseCS::GetFunctionString(): Unexpected InType = %d."), (int32)InType);
+		return TEXT("FElementWiseCS::GetFunctionString(): Unexpected InType.");
+	}
+	const TArray<FString> FunctionStrings({
+		TEXT("abs(X)"),						// Abs
+		TEXT("acos(X)"),					// Acos
+		TEXT("asin(X)"),					// Asin
+		TEXT("atan(X)"),					// Atan
+		TEXT("ceil(X)"),					// Ceil
+		TEXT("cos(X)"),						// Cos
+		TEXT("cosh(X)"),					// Cosh
+		TEXT("exp(X)"),						// Exp
+		TEXT("floor(X)"),					// Floor
+		TEXT("(X < 0 ? Attribute * X : X)"),// LeakyRelu
+		TEXT("log(X)"),						// Log
+		TEXT("-(X)"),						// Neg
+		TEXT("1/(X)"),						// Reciprocal
+		TEXT("round(X)"),					// Round
+		TEXT("(X < 0 ? 0 : X)"),			// Relu
+		TEXT("1/(1+exp(-(X)))"),			// Sigmoid
+		TEXT("sign(X)"),					// Sign
+		TEXT("sin(X)"),						// Sin
+		TEXT("sinh(X)"),					// Sinh
+		TEXT("sqrt(X)"),					// Sqrt
+		TEXT("tan(X)"),						// Tan
+		TEXT("tanh(X)") });					// Tanh
+	return FunctionStrings[(int32)InType];
+}
+
+uint32 FElementWiseCS::GetNumberAttributes(const EElementWiseOperator InType)
+{
+	if (InType >= EElementWiseOperator::MAX)
+	{
+		UE_LOG(LogNeuralNetworkInferenceShaders, Warning, TEXT("FElementWiseCS::GetNumberAttributes(): Unexpected InType = %d."), (int32)InType);
+		return 0;
+	}
+	const TArray<uint32> NumberAttributes({
+		0,	// Abs
+		0,	// Acos
+		0,	// Asin
+		0,	// Atan
+		0,	// Ceil
+		0,	// Cos
+		0,	// Cosh
+		0,	// Exp
+		0,	// Floor
+		1,	// LeakyRelu
+		0,	// Log
+		0,	// Neg
+		0,	// Reciprocal
+		0,	// Round
+		0,	// Relu
+		0,	// Sigmoid
+		0,	// Sign
+		0,	// Sin
+		0,	// Sinh
+		0,	// Sqrt
+		0,	// Tan
+		0 });// Tanh
+	return NumberAttributes[(int32)InType];
+}
+
+
+
+/* Shader implementation
+ *****************************************************************************/
+
+IMPLEMENT_SHADER_TYPE(, FElementWiseCS, TEXT("/Plugins/NeuralNetworkInference/Private/ElementWiseOperator.usf"), TEXT("ElementWiseCS"), SF_Compute) // Path defined in NeuralNetworkInferenceShadersModule.cpp
diff --git a/Engine/Plugins/Experimental/NeuralNetworkInference/Source/Private/NeuralNetworkInferenceShaders/Private/NeuralNetworkInferenceShaders/GemmCS.cpp b/Engine/Plugins/Experimental/NeuralNetworkInference/Source/Private/NeuralNetworkInferenceShaders/Private/NeuralNetworkInferenceShaders/GemmCS.cpp
new file mode 100644
--- /dev/null
+++ b/Engine/Plugins/Experimental/NeuralNetworkInference/Source/Private/NeuralNetworkInferenceShaders/Private/NeuralNetworkInferenceShaders/GemmCS.cpp
@@ -0,0 +1,27 @@
+// Copyright Epic Games, Inc. All Rights Reserved.
+
+#include "NeuralNetworkInferenceShaders/GemmCS.h"
+
+
+
+/* FGemmCS public functions
+ *****************************************************************************/
+
+void FGemmCS::ModifyCompilationEnvironment(const FGlobalShaderPermutationParameters& InParameters, FShaderCompilerEnvironment& OutEnvironment)
+{
+	FGlobalShader::ModifyCompilationEnvironment(InParameters, OutEnvironment);
+	OutEnvironment.SetDefine(TEXT("THREADGROUP_SIZE_X"), THREADGROUP_SIZE_X);
+	OutEnvironment.SetDefine(TEXT("THREADGROUP_SIZE_Y"), THREADGROUP_SIZE_Y);
+}
+
+
+
+const uint32 FGemmCS::THREADGROUP_SIZE_X = 128;
+const uint32 FGemmCS::THREADGROUP_SIZE_Y = 1;
+
+
+
+/* Shader implementation
+ *****************************************************************************/
+
+IMPLEMENT_SHADER_TYPE(, FGemmCS, TEXT("/Plugins/NeuralNetworkInference/Private/GemmOperator.usf"), TEXT("GemmCS"), SF_Compute) // Path defined in NeuralNetworkInferenceShadersModule.cpp
diff --git a/Engine/Plugins/Experimental/NeuralNetworkInference/Source/Private/NeuralNetworkInferenceShaders/Private/NeuralNetworkInferenceShaders/Module.cpp b/Engine/Plugins/Experimental/NeuralNetworkInference/Source/Private/NeuralNetworkInferenceShaders/Private/NeuralNetworkInferenceShaders/Module.cpp
new file mode 100644
--- /dev/null
+++ b/Engine/Plugins/Experimental/NeuralNetworkInference/Source/Private/NeuralNetworkInferenceShaders/Private/NeuralNetworkInferenceShaders/Module.cpp
@@ -0,0 +1,36 @@
+// Copyright Epic Games, Inc. All Rights Reserved.
+
+#include "NeuralNetworkInferenceShaders/Module.h"
+#include "NeuralNetworkInferenceShaders/Utils.h"
+#include "Interfaces/IPluginManager.h"
+#include "ShaderCore.h"
+
+#define LOCTEXT_NAMESPACE "NeuralNetworkInferenceShaders"
+
+// This code will execute after your module is loaded into memory; the exact timing is specified in the .uplugin file per-module
+void FNeuralNetworkInferenceShadersModule::StartupModule()
+{
+	// Add shader directory
+	const TSharedPtr<IPlugin> NeuralNetworkInferencePlugin = IPluginManager::Get().FindPlugin(TEXT("NeuralNetworkInference"));
+	if (NeuralNetworkInferencePlugin.IsValid())
+	{
+		const FString RealShaderDirectory = NeuralNetworkInferencePlugin->GetBaseDir() + TEXT("/Shaders/"); // TEXT("../../../Engine/Plugins/NeuralNetworkInference/Shaders");
+		const FString VirtualShaderDirectory = TEXT("/Plugins/NeuralNetworkInference");
+		AddShaderSourceDirectoryMapping(VirtualShaderDirectory, RealShaderDirectory);
+	}
+	else
+	{
+		UE_LOG(LogNeuralNetworkInferenceShaders, Warning,
+			TEXT("FNeuralNetworkInferenceShadersModule::StartupModule(): NeuralNetworkInferencePlugin was nullptr, shaders directory not added."));
+	}
+}
+
+// This function may be called during shutdown to clean up your module. For modules that support dynamic reloading,
+// we call this function before unloading the module.
+void FNeuralNetworkInferenceShadersModule::ShutdownModule()
+{
+}
+
+IMPLEMENT_MODULE(FNeuralNetworkInferenceShadersModule, NeuralNetworkInferenceShaders);
+
+#undef LOCTEXT_NAMESPACE
diff --git a/Engine/Plugins/Experimental/NeuralNetworkInference/Source/Private/NeuralNetworkInferenceShaders/Private/NeuralNetworkInferenceShaders/Module.h b/Engine/Plugins/Experimental/NeuralNetworkInference/Source/Private/NeuralNetworkInferenceShaders/Private/NeuralNetworkInferenceShaders/Module.h
new file mode 100644
--- /dev/null
+++ b/Engine/Plugins/Experimental/NeuralNetworkInference/Source/Private/NeuralNetworkInferenceShaders/Private/NeuralNetworkInferenceShaders/Module.h
@@ -0,0 +1,14 @@
+// Copyright Epic Games, Inc. All Rights Reserved.
+
+#pragma once
+
+#include "CoreMinimal.h"
+#include "Modules/ModuleManager.h"
+
+class FNeuralNetworkInferenceShadersModule : public IModuleInterface
+{
+public:
+	// Begin IModuleInterface
+	virtual void StartupModule() override;
+	virtual void ShutdownModule() override;
+};
diff --git a/Engine/Plugins/Experimental/NeuralNetworkInference/Source/Private/NeuralNetworkInferenceShaders/Private/NeuralNetworkInferenceShaders/MultidirectionalBroadcastCS.cpp b/Engine/Plugins/Experimental/NeuralNetworkInference/Source/Private/NeuralNetworkInferenceShaders/Private/NeuralNetworkInferenceShaders/MultidirectionalBroadcastCS.cpp
new file mode 100644
--- /dev/null
+++ b/Engine/Plugins/Experimental/NeuralNetworkInference/Source/Private/NeuralNetworkInferenceShaders/Private/NeuralNetworkInferenceShaders/MultidirectionalBroadcastCS.cpp
@@ -0,0 +1,58 @@
+// Copyright Epic Games, Inc. All Rights Reserved.
+
+#include "NeuralNetworkInferenceShaders/MultidirectionalBroadcastCS.h"
+#include "NeuralNetworkInferenceShaders/Utils.h"
+
+
+
+/* FMultidirectionalBroadcastCS static members
+ *****************************************************************************/
+
+const uint32 FMultidirectionalBroadcastCS::THREADGROUP_SIZE_X(128);
+const uint32 FMultidirectionalBroadcastCS::MAX_NUMBER_DIMENSIONS(16);
+
+
+
+/* FMultidirectionalBroadcastCS public functions
+ *****************************************************************************/
+
+void FMultidirectionalBroadcastCS::ModifyCompilationEnvironment(const FGlobalShaderPermutationParameters& InParameters, FShaderCompilerEnvironment& OutEnvironment)
+{
+	FGlobalShader::ModifyCompilationEnvironment(InParameters, OutEnvironment);
+	OutEnvironment.SetDefine(TEXT("THREADGROUP_SIZE_X"), THREADGROUP_SIZE_X);
+	OutEnvironment.SetDefine(TEXT("MAX_NUMBER_DIMENSIONS"), MAX_NUMBER_DIMENSIONS);
+
+	FPermutationDomain PermutationVector(InParameters.PermutationId);
+	const FString FunctionString = GetFunctionString(PermutationVector.Get<FShaderType>());
+	OutEnvironment.SetDefine(TEXT("MultidirectionalBroadcastFunction(X, Y)"), *FunctionString);
+
+	// Note: INLINED_MODE and SHAPE_MODE are already defined because we used SHADER_PERMUTATION_ENUM_CLASS with those names
+}
+
+
+
+/* FMultidirectionalBroadcastCS private static functions
+ *****************************************************************************/
+
+FString FMultidirectionalBroadcastCS::GetFunctionString(const EMultidirectionalBroadcastOperator InType)
+{
+	if (InType >= EMultidirectionalBroadcastOperator::MAX)
+	{
+		UE_LOG(LogNeuralNetworkInferenceShaders, Warning, TEXT("FMultidirectionalBroadcastCS::GetFunctionString(): Unexpected InType = %d."), (int32)InType);
+		return TEXT("FMultidirectionalBroadcastCS::GetFunctionString(): Unexpected InType.");
+	}
+	const TArray<FString> FunctionStrings({
+		TEXT("(X + Y)"),	// Add
+		TEXT("(X / Y)"),	// Div
+		TEXT("(X * Y)"),	// Mul
+		TEXT("pow(X, Y)"),	// Pow
+		TEXT("(X - Y)")});	// Sub
+	return FunctionStrings[(int32)InType];
+}
+
+
+
+/* Shader implementation
+ *****************************************************************************/
+
+IMPLEMENT_SHADER_TYPE(, FMultidirectionalBroadcastCS, TEXT("/Plugins/NeuralNetworkInference/Private/MultidirectionalBroadcastOperator.usf"), TEXT("MultidirectionalBroadcastCS"), SF_Compute) // Path defined in NeuralNetworkInferenceShadersModule.cpp
diff --git a/Engine/Plugins/Experimental/NeuralNetworkInference/Source/Private/NeuralNetworkInferenceShaders/Private/NeuralNetworkInferenceShaders/Utils.h b/Engine/Plugins/Experimental/NeuralNetworkInference/Source/Private/NeuralNetworkInferenceShaders/Private/NeuralNetworkInferenceShaders/Utils.h
new file mode 100644
--- /dev/null
+++ b/Engine/Plugins/Experimental/NeuralNetworkInference/Source/Private/NeuralNetworkInferenceShaders/Private/NeuralNetworkInferenceShaders/Utils.h
@@ -0,0 +1,7 @@
+// Copyright Epic Games, Inc. All Rights Reserved.
+
+#pragma once
+
+#include "CoreMinimal.h"
+
+DEFINE_LOG_CATEGORY_STATIC(LogNeuralNetworkInferenceShaders, Display, All);
diff --git a/Engine/Plugins/Experimental/NeuralNetworkInference/Source/Private/ThirdPartyHelperAndDLLLoader/Internal/ORTExceptionHandler.h b/Engine/Plugins/Experimental/NeuralNetworkInference/Source/Private/ThirdPartyHelperAndDLLLoader/Internal/ORTExceptionHandler.h
new file mode 100644
--- /dev/null
+++ b/Engine/Plugins/Experimental/NeuralNetworkInference/Source/Private/ThirdPartyHelperAndDLLLoader/Internal/ORTExceptionHandler.h
@@ -0,0 +1,12 @@
+// Copyright Epic Games, Inc. All Rights Reserved.
+
+#pragma once
+
+#include "CoreMinimal.h"
+#include <string>
+
+class THIRDPARTYHELPERANDDLLLOADER_API FORTExceptionHandler
+{
+public:
+	static void ThrowPseudoException(const std::string& InString, const int32 InCode);
+};
diff --git a/Engine/Plugins/Experimental/NeuralNetworkInference/Source/Private/ThirdPartyHelperAndDLLLoader/Internal/RedirectCoutAndCerrToUeLog.h b/Engine/Plugins/Experimental/NeuralNetworkInference/Source/Private/ThirdPartyHelperAndDLLLoader/Internal/RedirectCoutAndCerrToUeLog.h
new file mode 100644
--- /dev/null
+++ b/Engine/Plugins/Experimental/NeuralNetworkInference/Source/Private/ThirdPartyHelperAndDLLLoader/Internal/RedirectCoutAndCerrToUeLog.h
@@ -0,0 +1,37 @@
+// Copyright Epic Games, Inc. All Rights Reserved.
+
+#pragma once
+
+#include "CoreMinimal.h"
+#include <iostream>
+#include <sstream>
+
+/**
+ * Redirect std::cout to UE_LOG (since this class is constructed until it is destructed/goes out of scope).
+ * Undefined behavior (i.e., potential crashes) if this class is constructed more than once simultaneously.
+ */
+class THIRDPARTYHELPERANDDLLLOADER_API FRedirectCoutAndCerrToUeLog
+{
+public:
+	FRedirectCoutAndCerrToUeLog();
+	~FRedirectCoutAndCerrToUeLog();
+
+private:
+	std::streambuf* BackupStreamReadBufCout;
+	std::streambuf* BackupStreamReadBufCerr;
+
+	class LStream : public std::stringbuf
+	{
+	public:
+		LStream(const bool bInTrueIfCoutFalseIfCerr);
+
+	protected:
+		int sync();
+
+	private:
+		const bool bTrueIfCoutFalseIfCerr;
+	};
+
+	LStream StreamCout;
+	LStream StreamCerr;
+};
diff --git a/Engine/Plugins/Experimental/NeuralNetworkInference/Source/Private/ThirdPartyHelperAndDLLLoader/Internal/ThirdPartyWarningDisabler.h b/Engine/Plugins/Experimental/NeuralNetworkInference/Source/Private/ThirdPartyHelperAndDLLLoader/Internal/ThirdPartyWarningDisabler.h
new file mode 100644
--- /dev/null
+++ b/Engine/Plugins/Experimental/NeuralNetworkInference/Source/Private/ThirdPartyHelperAndDLLLoader/Internal/ThirdPartyWarningDisabler.h
@@ -0,0 +1,34 @@
+// Copyright Epic Games, Inc. All Rights Reserved.
+
+#pragma once
+
+#include "CoreMinimal.h"
+
+/*
+Many third party headers require some care when importing. NNI third party includes should be wrapped like this:
+	#include "ThirdPartyWarningDisabler.h"
+	NNI_THIRD_PARTY_INCLUDES_START
+	#undef check
+	#undef TEXT
+	// your ONNXRUNTIME include directives go here...
+	NNI_THIRD_PARTY_INCLUDES_END
+*/
+
+#ifdef PLATFORM_NNI_MICROSOFT
+#define NNI_THIRD_PARTY_INCLUDES_START THIRD_PARTY_INCLUDES_START \
+	__pragma(warning(disable: 4191)) /* C4191: 'reinterpret_cast': unsafe conversion from unsafe conversion from 'X' to 'Y' */ \
+	__pragma(warning(disable: 4497)) /* C4497: nonstandard extension 'sealed' used: replace with 'final' */ \
+	__pragma(warning(disable: 6001)) /* C6001: Using uninitialized memory '*pNode'. */ \
+	__pragma(warning(disable: 6294)) /* C6294: Ill-defined for-loop:  initial condition does not satisfy test.  Loop body not executed. */ \
+	__pragma(warning(disable: 6313)) /* C6313: Incorrect operator:  zero-valued flag cannot be tested with bitwise-and.  Use an equality test to check for zero-valued flags. */ \
+	__pragma(warning(disable: 6387)) /* C6387: 'X' could be '0':  this does not adhere to the specification for the function 'memset'. */ \
+	__pragma(warning(disable: 6388)) /* C6388: '*X' might not be '0':  this does not adhere to the specification for the function 'Y'. */ \
+	__pragma(warning(disable: 28020)) /* C28020: The expression '0<=_Param_(1)&&_Param_(1)<=400-1' is not true at this call. */ \
+	UE_PUSH_MACRO("check") \
+	UE_PUSH_MACRO("TEXT")
+#else
+// If support added for other platforms, this definition may require updating
+#define NNI_THIRD_PARTY_INCLUDES_START THIRD_PARTY_INCLUDES_START UE_PUSH_MACRO("check") UE_PUSH_MACRO("TEXT")
+#endif //PLATFORM_NNI_MICROSOFT
+
+#define NNI_THIRD_PARTY_INCLUDES_END THIRD_PARTY_INCLUDES_END UE_POP_MACRO("check") UE_POP_MACRO("TEXT")
diff --git a/Engine/Plugins/Experimental/NeuralNetworkInference/Source/Private/ThirdPartyHelperAndDLLLoader/Private/ORTExceptionHandler.cpp b/Engine/Plugins/Experimental/NeuralNetworkInference/Source/Private/ThirdPartyHelperAndDLLLoader/Private/ORTExceptionHandler.cpp
new file mode 100644
--- /dev/null
+++ b/Engine/Plugins/Experimental/NeuralNetworkInference/Source/Private/ThirdPartyHelperAndDLLLoader/Private/ORTExceptionHandler.cpp
@@ -0,0 +1,17 @@
+// Copyright Epic Games, Inc. All Rights Reserved.
+
+#include "ORTExceptionHandler.h"
+#include "ThirdPartyHelperAndDLLLoaderUtils.h"
+
+
+
+/* FORTExceptionHandler public classes
+ *****************************************************************************/
+
+void FORTExceptionHandler::ThrowPseudoException(const std::string& InString, const int32 InCode)
+{
+	const FString ErrorMessage = FString::Format(TEXT("ONNXRuntime threw an exception with code {0}, e.what(): \"{1}\"."),
+		{ FString::FromInt(InCode), FString(ANSI_TO_TCHAR(InString.c_str())) });
+	UE_LOG(LogNeuralNetworkInferenceThirdPartyHelperAndDLLLoader, Warning, TEXT("%s"), *ErrorMessage);
+	checkf(false, TEXT("%s"), *ErrorMessage);
+}
diff --git a/Engine/Plugins/Experimental/NeuralNetworkInference/Source/Private/ThirdPartyHelperAndDLLLoader/Private/RedirectCoutAndCerrToUeLog.cpp b/Engine/Plugins/Experimental/NeuralNetworkInference/Source/Private/ThirdPartyHelperAndDLLLoader/Private/RedirectCoutAndCerrToUeLog.cpp
new file mode 100644
--- /dev/null
+++ b/Engine/Plugins/Experimental/NeuralNetworkInference/Source/Private/ThirdPartyHelperAndDLLLoader/Private/RedirectCoutAndCerrToUeLog.cpp
@@ -0,0 +1,78 @@
+// Copyright Epic Games, Inc. All Rights Reserved.
+
+#include "RedirectCoutAndCerrToUeLog.h"
+#include "ThirdPartyHelperAndDLLLoaderUtils.h"
+
+
+
+/* FRedirectCoutAndCerrToUeLog public classes
+ *****************************************************************************/
+
+FRedirectCoutAndCerrToUeLog::FRedirectCoutAndCerrToUeLog()
+	: StreamCout(true)
+	, StreamCerr(false)
+{
+	// Save pointer to std::cout buffer
+	BackupStreamReadBufCout = std::cout.rdbuf();
+	BackupStreamReadBufCerr = std::cerr.rdbuf();
+	// Substitute internal std::cout/std::cerr buffer with
+	std::cout.rdbuf(&StreamCout); // now std::cout work with 'StreamCout' buffer
+	std::cerr.rdbuf(&StreamCerr); // now std::cerr work with 'StreamCerr' buffer
+}
+
+FRedirectCoutAndCerrToUeLog::~FRedirectCoutAndCerrToUeLog()
+{
+	// Go back to old buffers
+	std::cout.rdbuf(BackupStreamReadBufCout);
+	std::cerr.rdbuf(BackupStreamReadBufCerr);
+}
+
+FRedirectCoutAndCerrToUeLog::LStream::LStream(const bool bInTrueIfCoutFalseIfCerr)
+	: bTrueIfCoutFalseIfCerr(bInTrueIfCoutFalseIfCerr)
+{
+}
+
+int FRedirectCoutAndCerrToUeLog::LStream::sync()
+{
+	FString StdMessage = FString(str().c_str());
+	// Print message
+	if (StdMessage.Len() > 0)
+	{
+		// Remove break lines and/or spaces at the end of the message
+		if (StdMessage.Len() > 1)
+		{
+			if (StdMessage.Right(1) == TEXT("\n") || StdMessage.Right(1) == TEXT(" "))
+			{
+				// Find how many break lines and/or spaces at the end
+				int64 BreakLineCounter = 1;
+				while (StdMessage.Right(BreakLineCounter) == TEXT("\n") || StdMessage.Right(BreakLineCounter) == TEXT(" "))
+				{
+					++BreakLineCounter;
+					// E.g., If message = \n\n\n --> We want to keep only 1 \n
+					if (BreakLineCounter >= StdMessage.Len())
+					{
+						BreakLineCounter = StdMessage.Len() - 1;
+						break;
+					}
+				}
+				// Remove those final break lines and/or spaces
+				if (StdMessage.Len() > BreakLineCounter)
+				{
+					StdMessage.LeftChopInline(BreakLineCounter - 1);
+				}
+			}
+		}
+		// Print cleaned message
+		if (bTrueIfCoutFalseIfCerr)
+		{
+			UE_LOG(LogNeuralNetworkInferenceThirdPartyHelperAndDLLLoader, Display, TEXT("std::cout: %s"), *StdMessage);
+		}
+		else
+		{
+			UE_LOG(LogNeuralNetworkInferenceThirdPartyHelperAndDLLLoader, Warning, TEXT("std::cerr: %s"), *StdMessage);
+		}
+	}
+	// Reset buffer (otherwise it'd keep printing the old messages together with the new ones)
+	str("");
+	return std::stringbuf::sync();
+}
diff --git a/Engine/Plugins/Experimental/NeuralNetworkInference/Source/Private/ThirdPartyHelperAndDLLLoader/Private/ThirdPartyHelperAndDLLLoaderModule.cpp b/Engine/Plugins/Experimental/NeuralNetworkInference/Source/Private/ThirdPartyHelperAndDLLLoader/Private/ThirdPartyHelperAndDLLLoaderModule.cpp
new file mode 100644
--- /dev/null
+++ b/Engine/Plugins/Experimental/NeuralNetworkInference/Source/Private/ThirdPartyHelperAndDLLLoader/Private/ThirdPartyHelperAndDLLLoaderModule.cpp
@@ -0,0 +1,34 @@
+// Copyright Epic Games, Inc. All Rights Reserved.
+
+#include "ThirdPartyHelperAndDLLLoaderModule.h"
+#include "HAL/PlatformProcess.h"
+#include "Interfaces/IPluginManager.h"
+
+// This code will execute after your module is loaded into memory; the exact timing is specified in the .uplugin file per-module
+void FThirdPartyHelperAndDLLLoaderModule::StartupModule()
+{
+#ifdef PLATFORM_WIN64
+	const FString PluginDir = IPluginManager::Get().FindPlugin(TEXT("NeuralNetworkInference"))->GetBaseDir();
+	const FString DirectMLRuntimeBinPath = PluginDir / TEXT(PREPROCESSOR_TO_STRING(DIRECTML_PLATFORM_PATH));
+	const FString DirectMLDLLPath = DirectMLRuntimeBinPath / TEXT(PREPROCESSOR_TO_STRING(DIRECTML_DLL_NAME));
+
+	FPlatformProcess::PushDllDirectory(*DirectMLRuntimeBinPath);
+	DirectMLDLLHandle = FPlatformProcess::GetDllHandle(*DirectMLDLLPath);
+	FPlatformProcess::PopDllDirectory(*DirectMLRuntimeBinPath);
+#endif //PLATFORM_WIN64
+}
+
+// This function may be called during shutdown to clean up your module. For modules that support dynamic reloading,
+// we call this function before unloading the module.
+void FThirdPartyHelperAndDLLLoaderModule::ShutdownModule()
+{
+#ifdef PLATFORM_WIN64
+	if (DirectMLDLLHandle)
+	{
+		FPlatformProcess::FreeDllHandle(DirectMLDLLHandle);
+		DirectMLDLLHandle = nullptr;
+	}
+#endif //PLATFORM_WIN64
+}
+
+IMPLEMENT_MODULE(FThirdPartyHelperAndDLLLoaderModule, ThirdPartyHelperAndDLLLoader);
diff --git a/Engine/Plugins/Experimental/NeuralNetworkInference/Source/Private/ThirdPartyHelperAndDLLLoader/Private/ThirdPartyHelperAndDLLLoaderModule.h b/Engine/Plugins/Experimental/NeuralNetworkInference/Source/Private/ThirdPartyHelperAndDLLLoader/Private/ThirdPartyHelperAndDLLLoaderModule.h
new file mode 100644
--- /dev/null
+++ b/Engine/Plugins/Experimental/NeuralNetworkInference/Source/Private/ThirdPartyHelperAndDLLLoader/Private/ThirdPartyHelperAndDLLLoaderModule.h
@@ -0,0 +1,20 @@
+// Copyright Epic Games, Inc. All Rights Reserved.
+
+#pragma once
+
+#include "CoreMinimal.h"
+#include "Modules/ModuleManager.h"
+
+/**
+* The public interface to this module
+*/
+class FThirdPartyHelperAndDLLLoaderModule : public IModuleInterface
+{
+public:
+	// Begin IModuleInterface
+	virtual void StartupModule() override;
+	virtual void ShutdownModule() override;
+
+private:
+	void* DirectMLDLLHandle;
+};
diff --git a/Engine/Plugins/Experimental/NeuralNetworkInference/Source/Private/ThirdPartyHelperAndDLLLoader/Private/ThirdPartyHelperAndDLLLoaderUtils.h b/Engine/Plugins/Experimental/NeuralNetworkInference/Source/Private/ThirdPartyHelperAndDLLLoader/Private/ThirdPartyHelperAndDLLLoaderUtils.h
new file mode 100644
--- /dev/null
+++ b/Engine/Plugins/Experimental/NeuralNetworkInference/Source/Private/ThirdPartyHelperAndDLLLoader/Private/ThirdPartyHelperAndDLLLoaderUtils.h
@@ -0,0 +1,7 @@
+// Copyright Epic Games, Inc. All Rights Reserved.
+
+#pragma once
+
+#include "CoreMinimal.h"
+
+DEFINE_LOG_CATEGORY_STATIC(LogNeuralNetworkInferenceThirdPartyHelperAndDLLLoader, Display, All);
diff --git a/Engine/Plugins/Experimental/NeuralNetworkInference/Source/Private/ThirdPartyHelperAndDLLLoader/ThirdPartyHelperAndDLLLoader.Build.cs b/Engine/Plugins/Experimental/NeuralNetworkInference/Source/Private/ThirdPartyHelperAndDLLLoader/ThirdPartyHelperAndDLLLoader.Build.cs
new file mode 100644
--- /dev/null
+++ b/Engine/Plugins/Experimental/NeuralNetworkInference/Source/Private/ThirdPartyHelperAndDLLLoader/ThirdPartyHelperAndDLLLoader.Build.cs
@@ -0,0 +1,39 @@
+// Copyright Epic Games, Inc. All Rights Reserved.
+
+using UnrealBuildTool;
+
+public class ThirdPartyHelperAndDLLLoader : ModuleRules
+{
+	public ThirdPartyHelperAndDLLLoader( ReadOnlyTargetRules Target ) : base( Target )
+	{
+		PCHUsage = ModuleRules.PCHUsageMode.UseExplicitOrSharedPCHs;
+
+		PublicIncludePaths.AddRange(
+			new string[] {
+				System.IO.Path.Combine(ModuleDirectory, "..")
+			}
+		);
+
+		PublicDependencyModuleNames.AddRange
+			(
+			new string[] {
+				"Core"
+			}
+		);
+
+		PrivateDependencyModuleNames.AddRange
+			(
+			new string[] {
+				"DirectML",
+				"Projects"
+			}
+		);
+
+		// Win64-only
+		if (Target.Platform == UnrealTargetPlatform.Win64)
+		{
+			PublicDefinitions.Add("PLATFORM_WIN64");
+			PublicDefinitions.Add("PLATFORM_NNI_MICROSOFT");
+		}
+	}
+}
