ModelProto.h only used by Private classes (no longer by internal ones).
Also removed not required include modelproto in squeezeOp

_codereview Francisco.Vicente
#rnx

diff --git a/Engine/Plugins/Experimental/NeuralNetworkInference/Source/NeuralNetworkInference/Internal/NeuralOperators/BatchNormalizationOperator.h b/Engine/Plugins/Experimental/NeuralNetworkInference/Source/NeuralNetworkInference/Internal/NeuralOperators/BatchNormalizationOperator.h
--- a/Engine/Plugins/Experimental/NeuralNetworkInference/Source/NeuralNetworkInference/Internal/NeuralOperators/BatchNormalizationOperator.h
+++ b/Engine/Plugins/Experimental/NeuralNetworkInference/Source/NeuralNetworkInference/Internal/NeuralOperators/BatchNormalizationOperator.h
@@ -3,13 +3,12 @@
 #pragma once
 
 #include "CoreMinimal.h"
-#include "ModelProto.h"
 #include "NeuralOperator.h"
 
 class NEURALNETWORKINFERENCE_API FBatchNormalizationOperator : public FNeuralOperator
 {
 public:
-	FBatchNormalizationOperator(const bool bIsInlinedTensor, const FNodeProto& InNodeProto);
+	FBatchNormalizationOperator(const bool bIsInlinedTensor, const struct FNodeProto* const InNodeProto);
 
 	FBatchNormalizationOperator(const bool bIsInlinedTensor, const float InEpsilon = 1e-5, const float InMomentum = 0.9f);
 
diff --git a/Engine/Plugins/Experimental/NeuralNetworkInference/Source/NeuralNetworkInference/Internal/NeuralOperators/ConstantOperator.h b/Engine/Plugins/Experimental/NeuralNetworkInference/Source/NeuralNetworkInference/Internal/NeuralOperators/ConstantOperator.h
--- a/Engine/Plugins/Experimental/NeuralNetworkInference/Source/NeuralNetworkInference/Internal/NeuralOperators/ConstantOperator.h
+++ b/Engine/Plugins/Experimental/NeuralNetworkInference/Source/NeuralNetworkInference/Internal/NeuralOperators/ConstantOperator.h
@@ -3,13 +3,12 @@
 #pragma once
 
 #include "CoreMinimal.h"
-#include "ModelProto.h"
 #include "NeuralOperator.h"
 
 class NEURALNETWORKINFERENCE_API FConstantOperator : public FNeuralOperator
 {
 public:
-	FConstantOperator(const FNodeProto& InNodeProto);
+	FConstantOperator(const struct FNodeProto* const InNodeProto);
 
 	FConstantOperator(const FNeuralTensor& InTensor);
 
diff --git a/Engine/Plugins/Experimental/NeuralNetworkInference/Source/NeuralNetworkInference/Internal/NeuralOperators/ConvOperator.h b/Engine/Plugins/Experimental/NeuralNetworkInference/Source/NeuralNetworkInference/Internal/NeuralOperators/ConvOperator.h
--- a/Engine/Plugins/Experimental/NeuralNetworkInference/Source/NeuralNetworkInference/Internal/NeuralOperators/ConvOperator.h
+++ b/Engine/Plugins/Experimental/NeuralNetworkInference/Source/NeuralNetworkInference/Internal/NeuralOperators/ConvOperator.h
@@ -2,13 +2,12 @@
 
 #pragma once
 
-#include "ModelProto.h"
 #include "NeuralOperators/ConvBaseOperator.h"
 
 class NEURALNETWORKINFERENCE_API FConvOperator : public FConvBaseOperator
 {
 public:
-	FConvOperator(const FNodeProto& InNodeProto);
+	FConvOperator(const struct FNodeProto* const InNodeProto);
 
 	FConvOperator(const EAutoPad InAutoPad, const TArray<int64>& InDilations, const int64 InGroup, const TArray<int64>& InKernelShape, const TArray<int64>& InPads, const TArray<int64>& InStrides);
 
diff --git a/Engine/Plugins/Experimental/NeuralNetworkInference/Source/NeuralNetworkInference/Internal/NeuralOperators/ConvTransposeOperator.h b/Engine/Plugins/Experimental/NeuralNetworkInference/Source/NeuralNetworkInference/Internal/NeuralOperators/ConvTransposeOperator.h
--- a/Engine/Plugins/Experimental/NeuralNetworkInference/Source/NeuralNetworkInference/Internal/NeuralOperators/ConvTransposeOperator.h
+++ b/Engine/Plugins/Experimental/NeuralNetworkInference/Source/NeuralNetworkInference/Internal/NeuralOperators/ConvTransposeOperator.h
@@ -2,14 +2,13 @@
 
 #pragma once
 
-#include "ModelProto.h"
 #include "NeuralInt64ArrayUInt32Buffer.h"
 #include "NeuralOperators/ConvBaseOperator.h"
 
 class NEURALNETWORKINFERENCE_API FConvTransposeOperator : public FConvBaseOperator
 {
 public:
-	FConvTransposeOperator(const FNodeProto& InNodeProto);
+	FConvTransposeOperator(const struct FNodeProto* const InNodeProto);
 
 	FConvTransposeOperator(const EAutoPad InAutoPad, const TArray<int64>& InDilations, const int64 InGroup, const TArray<int64>& InKernelShape,
 		const TArray<int64>& InOutputPadding, const TArray<int64>& InOutputShape, const TArray<int64>& InPads, const TArray<int64>& InStrides);
diff --git a/Engine/Plugins/Experimental/NeuralNetworkInference/Source/NeuralNetworkInference/Internal/NeuralOperators/GemmOperator.h b/Engine/Plugins/Experimental/NeuralNetworkInference/Source/NeuralNetworkInference/Internal/NeuralOperators/GemmOperator.h
--- a/Engine/Plugins/Experimental/NeuralNetworkInference/Source/NeuralNetworkInference/Internal/NeuralOperators/GemmOperator.h
+++ b/Engine/Plugins/Experimental/NeuralNetworkInference/Source/NeuralNetworkInference/Internal/NeuralOperators/GemmOperator.h
@@ -3,13 +3,12 @@
 #pragma once
 
 #include "CoreMinimal.h"
-#include "ModelProto.h"
 #include "NeuralOperator.h"
 
 class NEURALNETWORKINFERENCE_API FGemmOperator : public FNeuralOperator
 {
 public:
-	FGemmOperator(const FNodeProto& InNodeProto);
+	FGemmOperator(const struct FNodeProto* const InNodeProto);
 
 	FGemmOperator(const float InAlpha = 1.f, const float InBeta = 1.f, const bool bInTransA = false, const bool bInTransB = false);
 
diff --git a/Engine/Plugins/Experimental/NeuralNetworkInference/Source/NeuralNetworkInference/Internal/NeuralOperators/LeakyReluOperator.h b/Engine/Plugins/Experimental/NeuralNetworkInference/Source/NeuralNetworkInference/Internal/NeuralOperators/LeakyReluOperator.h
--- a/Engine/Plugins/Experimental/NeuralNetworkInference/Source/NeuralNetworkInference/Internal/NeuralOperators/LeakyReluOperator.h
+++ b/Engine/Plugins/Experimental/NeuralNetworkInference/Source/NeuralNetworkInference/Internal/NeuralOperators/LeakyReluOperator.h
@@ -3,13 +3,12 @@
 #pragma once
 
 #include "CoreMinimal.h"
-#include "ModelProto.h"
 #include "NeuralOperators/ElementWiseOperator.h"
 
 class NEURALNETWORKINFERENCE_API FLeakyReluOperator : public IElementWiseOperator
 {
 public:
-	FLeakyReluOperator(const bool bIsInlinedTensor, const FNodeProto& InNodeProto);
+	FLeakyReluOperator(const bool bIsInlinedTensor, const struct FNodeProto* const InNodeProto);
 
 	FLeakyReluOperator(const bool bIsInlinedTensor, const float InAlpha = 0.01f);
 
diff --git a/Engine/Plugins/Experimental/NeuralNetworkInference/Source/NeuralNetworkInference/Internal/NeuralOperators/ReshapeOperator.h b/Engine/Plugins/Experimental/NeuralNetworkInference/Source/NeuralNetworkInference/Internal/NeuralOperators/ReshapeOperator.h
--- a/Engine/Plugins/Experimental/NeuralNetworkInference/Source/NeuralNetworkInference/Internal/NeuralOperators/ReshapeOperator.h
+++ b/Engine/Plugins/Experimental/NeuralNetworkInference/Source/NeuralNetworkInference/Internal/NeuralOperators/ReshapeOperator.h
@@ -3,7 +3,6 @@
 #pragma once
 
 #include "CoreMinimal.h"
-#include "ModelProto.h"
 #include "NeuralOperator.h"
 
 /**
@@ -13,7 +12,7 @@
 class NEURALNETWORKINFERENCE_API FReshapeOperator : public FNeuralOperator
 {
 public:
-	FReshapeOperator(const bool bIsInlinedTensor, const FNodeProto& InNodeProto);
+	FReshapeOperator(const bool bIsInlinedTensor, const struct FNodeProto* const InNodeProto);
 
 	FReshapeOperator(const bool bIsInlinedTensor, const int64 InAllowZero = 0);
 
diff --git a/Engine/Plugins/Experimental/NeuralNetworkInference/Source/NeuralNetworkInference/Internal/NeuralOperators/SqueezeOperator.h b/Engine/Plugins/Experimental/NeuralNetworkInference/Source/NeuralNetworkInference/Internal/NeuralOperators/SqueezeOperator.h
--- a/Engine/Plugins/Experimental/NeuralNetworkInference/Source/NeuralNetworkInference/Internal/NeuralOperators/SqueezeOperator.h
+++ b/Engine/Plugins/Experimental/NeuralNetworkInference/Source/NeuralNetworkInference/Internal/NeuralOperators/SqueezeOperator.h
@@ -3,7 +3,6 @@
 #pragma once
 
 #include "CoreMinimal.h"
-#include "ModelProto.h"
 #include "NeuralOperator.h"
 
 
diff --git a/Engine/Plugins/Experimental/NeuralNetworkInference/Source/NeuralNetworkInference/Private/NeuralOperators/BatchNormalizationOperator.cpp b/Engine/Plugins/Experimental/NeuralNetworkInference/Source/NeuralNetworkInference/Private/NeuralOperators/BatchNormalizationOperator.cpp
--- a/Engine/Plugins/Experimental/NeuralNetworkInference/Source/NeuralNetworkInference/Private/NeuralOperators/BatchNormalizationOperator.cpp
+++ b/Engine/Plugins/Experimental/NeuralNetworkInference/Source/NeuralNetworkInference/Private/NeuralOperators/BatchNormalizationOperator.cpp
@@ -1,6 +1,7 @@
 // Copyright Epic Games, Inc. All Rights Reserved.
 
 #include "NeuralOperators/BatchNormalizationOperator.h"
+#include "ModelProto.h"
 #include "NeuralNetworkInferenceShaders/BatchNormalizationCS.h"
 #include "NeuralNetworkInferenceUtils.h"
 #include "NeuralNetworkInferenceUtilsGPU.h"
@@ -37,14 +38,20 @@ void FPrivateBatchNormalizationOperator::ChannelBatchNormalization(FNeuralTensor
 /* FBatchNormalizationOperator structors
  *****************************************************************************/
 
-FBatchNormalizationOperator::FBatchNormalizationOperator(const bool bIsInlinedTensor, const FNodeProto& InNodeProto)
+FBatchNormalizationOperator::FBatchNormalizationOperator(const bool bIsInlinedTensor, const FNodeProto* const InNodeProto)
 	: FBatchNormalizationOperator(bIsInlinedTensor)
 {
-	if (const FAttributeProto* EpsilonAttribute = FModelProto::FindElementInArray(TEXT("Epsilon"), InNodeProto.Attribute, /*bMustValueBeFound*/false))
+	// Sanity check
+	if (!InNodeProto)
+	{
+		UE_LOG(LogNeuralNetworkInference, Warning, TEXT("FBatchNormalizationOperator(): InNodeProto was a nullptr."));
+		return;
+	}
+	if (const FAttributeProto* EpsilonAttribute = FModelProto::FindElementInArray(TEXT("Epsilon"), InNodeProto->Attribute, /*bMustValueBeFound*/false))
 	{
 		Epsilon = EpsilonAttribute->F;
 	}
-	if (const FAttributeProto* MomentumAttribute = FModelProto::FindElementInArray(TEXT("Momentum"), InNodeProto.Attribute, /*bMustValueBeFound*/false))
+	if (const FAttributeProto* MomentumAttribute = FModelProto::FindElementInArray(TEXT("Momentum"), InNodeProto->Attribute, /*bMustValueBeFound*/false))
 	{
 		Momentum = MomentumAttribute->F;
 	}
diff --git a/Engine/Plugins/Experimental/NeuralNetworkInference/Source/NeuralNetworkInference/Private/NeuralOperators/ConstantOperator.cpp b/Engine/Plugins/Experimental/NeuralNetworkInference/Source/NeuralNetworkInference/Private/NeuralOperators/ConstantOperator.cpp
--- a/Engine/Plugins/Experimental/NeuralNetworkInference/Source/NeuralNetworkInference/Private/NeuralOperators/ConstantOperator.cpp
+++ b/Engine/Plugins/Experimental/NeuralNetworkInference/Source/NeuralNetworkInference/Private/NeuralOperators/ConstantOperator.cpp
@@ -1,6 +1,7 @@
 // Copyright Epic Games, Inc. All Rights Reserved.
 
 #include "NeuralOperators/ConstantOperator.h"
+#include "ModelProto.h"
 #include "NeuralNetworkInferenceUtils.h"
 
 
@@ -8,15 +9,21 @@
 /* FConstantOperator structors
  *****************************************************************************/
 
-FConstantOperator::FConstantOperator(const FNodeProto& InNodeProto)
+FConstantOperator::FConstantOperator(const FNodeProto* const InNodeProto)
 	: FConstantOperator(FNeuralTensor())
 {
+	// Sanity check
+	if (!InNodeProto)
+	{
+		UE_LOG(LogNeuralNetworkInference, Warning, TEXT("FConstantOperator(): InNodeProto was a nullptr."));
+		return;
+	}
 	// Fill tensor from one of the attributes
-	if (const FAttributeProto* SparseValueAttribute = FModelProto::FindElementInArray(TEXT("sparse_value"), InNodeProto.Attribute, /*bMustValueBeFound*/false))
+	if (const FAttributeProto* SparseValueAttribute = FModelProto::FindElementInArray(TEXT("sparse_value"), InNodeProto->Attribute, /*bMustValueBeFound*/false))
 	{
 		UE_LOG(LogNeuralNetworkInference, Warning, TEXT("FConstantOperator::FConstantOperator(): \"sparse_value\" not implemented yet."));
 	}
-	else if (const FAttributeProto* ValueAttribute = FModelProto::FindElementInArray(TEXT("value"), InNodeProto.Attribute, /*bMustValueBeFound*/false))
+	else if (const FAttributeProto* ValueAttribute = FModelProto::FindElementInArray(TEXT("value"), InNodeProto->Attribute, /*bMustValueBeFound*/false))
 	{
 		if (ValueAttribute->Type == EAttributeProtoAttributeType::TENSOR)
 		{
@@ -32,27 +39,27 @@ FConstantOperator::FConstantOperator(const FNodeProto& InNodeProto)
 			UE_LOG(LogNeuralNetworkInference, Warning, TEXT("FConstantOperator::FConstantOperator(): Not implemented type used, ValueAttribute->Type = %d."), (int32)ValueAttribute->Type);
 		}
 	}
-	else if (const FAttributeProto* ValueFloatAttribute = FModelProto::FindElementInArray(TEXT("value_float"), InNodeProto.Attribute, /*bMustValueBeFound*/false))
+	else if (const FAttributeProto* ValueFloatAttribute = FModelProto::FindElementInArray(TEXT("value_float"), InNodeProto->Attribute, /*bMustValueBeFound*/false))
 	{
 		UE_LOG(LogNeuralNetworkInference, Warning, TEXT("FConstantOperator::FConstantOperator(): \"value\" not implemented yet."));
 	}
-	else if (const FAttributeProto* ValueFloatsAttribute = FModelProto::FindElementInArray(TEXT("value_floats"), InNodeProto.Attribute, /*bMustValueBeFound*/false))
+	else if (const FAttributeProto* ValueFloatsAttribute = FModelProto::FindElementInArray(TEXT("value_floats"), InNodeProto->Attribute, /*bMustValueBeFound*/false))
 	{
 		UE_LOG(LogNeuralNetworkInference, Warning, TEXT("FConstantOperator::FConstantOperator(): \"value_floats\" not implemented yet."));
 	}
-	else if (const FAttributeProto* ValueIntAttribute = FModelProto::FindElementInArray(TEXT("value_int"), InNodeProto.Attribute, /*bMustValueBeFound*/false))
+	else if (const FAttributeProto* ValueIntAttribute = FModelProto::FindElementInArray(TEXT("value_int"), InNodeProto->Attribute, /*bMustValueBeFound*/false))
 	{
 		UE_LOG(LogNeuralNetworkInference, Warning, TEXT("FConstantOperator::FConstantOperator(): \"value_int\" not implemented yet."));
 	}
-	else if (const FAttributeProto* ValueIntsAttribute = FModelProto::FindElementInArray(TEXT("value_ints"), InNodeProto.Attribute,/*bMustValueBeFound*/ false))
+	else if (const FAttributeProto* ValueIntsAttribute = FModelProto::FindElementInArray(TEXT("value_ints"), InNodeProto->Attribute,/*bMustValueBeFound*/ false))
 	{
 		UE_LOG(LogNeuralNetworkInference, Warning, TEXT("FConstantOperator::FConstantOperator(): \"value_ints\" not implemented yet."));
 	}
-	else if (const FAttributeProto* ValueStringAttribute = FModelProto::FindElementInArray(TEXT("value_string"), InNodeProto.Attribute, /*bMustValueBeFound*/false))
+	else if (const FAttributeProto* ValueStringAttribute = FModelProto::FindElementInArray(TEXT("value_string"), InNodeProto->Attribute, /*bMustValueBeFound*/false))
 	{
 		UE_LOG(LogNeuralNetworkInference, Warning, TEXT("FConstantOperator::FConstantOperator(): \"value_string\" not implemented yet."));
 	}
-	else if (const FAttributeProto* ValueStringsAttribute = FModelProto::FindElementInArray(TEXT("value_strings"), InNodeProto.Attribute, /*bMustValueBeFound*/false))
+	else if (const FAttributeProto* ValueStringsAttribute = FModelProto::FindElementInArray(TEXT("value_strings"), InNodeProto->Attribute, /*bMustValueBeFound*/false))
 	{
 		UE_LOG(LogNeuralNetworkInference, Warning, TEXT("FConstantOperator::FConstantOperator(): \"value_strings\" not implemented yet."));
 	}
diff --git a/Engine/Plugins/Experimental/NeuralNetworkInference/Source/NeuralNetworkInference/Private/NeuralOperators/ConvOperator.cpp b/Engine/Plugins/Experimental/NeuralNetworkInference/Source/NeuralNetworkInference/Private/NeuralOperators/ConvOperator.cpp
--- a/Engine/Plugins/Experimental/NeuralNetworkInference/Source/NeuralNetworkInference/Private/NeuralOperators/ConvOperator.cpp
+++ b/Engine/Plugins/Experimental/NeuralNetworkInference/Source/NeuralNetworkInference/Private/NeuralOperators/ConvOperator.cpp
@@ -1,6 +1,7 @@
 // Copyright Epic Games, Inc. All Rights Reserved.
 
 #include "NeuralOperators/ConvOperator.h"
+#include "ModelProto.h"
 #include "NeuralNetworkInferenceUtils.h"
 #include "NeuralNetworkInferenceUtilsGPU.h"
 
@@ -12,13 +13,19 @@
 class FPrivateConvOperator
 {
 public:
-	static FConvOperator NodeProtoToConvOperator(const FNodeProto& InNodeProto);
+	static FConvOperator NodeProtoToConvOperator(const FNodeProto* const InNodeProto);
 };
 
-FConvOperator FPrivateConvOperator::NodeProtoToConvOperator(const FNodeProto& InNodeProto)
+FConvOperator FPrivateConvOperator::NodeProtoToConvOperator(const FNodeProto* const InNodeProto)
 {
+	// Sanity check
+	if (!InNodeProto)
+	{
+		UE_LOG(LogNeuralNetworkInference, Warning, TEXT("FConvOperator(): InNodeProto was a nullptr."));
+		return FConvOperator(FConvBaseOperator::EAutoPad::NotSet, {}, -1, {}, {}, {});
+	}
 	FConvBaseOperator::EAutoPad AutoPad = FConvBaseOperator::EAutoPad::NotSet;
-	if (const FAttributeProto* EpsilonAttribute = FModelProto::FindElementInArray(TEXT("AutoPad"), InNodeProto.Attribute, /*bMustValueBeFound*/false))
+	if (const FAttributeProto* EpsilonAttribute = FModelProto::FindElementInArray(TEXT("AutoPad"), InNodeProto->Attribute, /*bMustValueBeFound*/false))
 	{
 		if (EpsilonAttribute->S == TEXT("NOTSET"))
 		{
@@ -42,12 +49,12 @@ FConvOperator FPrivateConvOperator::NodeProtoToConvOperator(const FNodeProto& In
 		}
 	}
 	TArray<int64> Dilations;
-	if (const FAttributeProto* MomentumAttribute = FModelProto::FindElementInArray(TEXT("Dilations"), InNodeProto.Attribute, /*bMustValueBeFound*/false))
+	if (const FAttributeProto* MomentumAttribute = FModelProto::FindElementInArray(TEXT("Dilations"), InNodeProto->Attribute, /*bMustValueBeFound*/false))
 	{
 		Dilations = MomentumAttribute->Integers;
 	}
 	int64 Group;
-	if (const FAttributeProto* MomentumAttribute = FModelProto::FindElementInArray(TEXT("Group"), InNodeProto.Attribute, /*bMustValueBeFound*/false))
+	if (const FAttributeProto* MomentumAttribute = FModelProto::FindElementInArray(TEXT("Group"), InNodeProto->Attribute, /*bMustValueBeFound*/false))
 	{
 		Group = MomentumAttribute->I;
 	}
@@ -56,17 +63,17 @@ FConvOperator FPrivateConvOperator::NodeProtoToConvOperator(const FNodeProto& In
 		Group = 1;
 	}
 	TArray<int64> KernelShape;
-	if (const FAttributeProto* MomentumAttribute = FModelProto::FindElementInArray(TEXT("KernelShape"), InNodeProto.Attribute, /*bMustValueBeFound*/false))
+	if (const FAttributeProto* MomentumAttribute = FModelProto::FindElementInArray(TEXT("KernelShape"), InNodeProto->Attribute, /*bMustValueBeFound*/false))
 	{
 		KernelShape = MomentumAttribute->Integers;
 	}
 	TArray<int64> Pads;
-	if (const FAttributeProto* MomentumAttribute = FModelProto::FindElementInArray(TEXT("Pads"), InNodeProto.Attribute, /*bMustValueBeFound*/false))
+	if (const FAttributeProto* MomentumAttribute = FModelProto::FindElementInArray(TEXT("Pads"), InNodeProto->Attribute, /*bMustValueBeFound*/false))
 	{
 		Pads = MomentumAttribute->Integers;
 	}
 	TArray<int64> Strides;
-	if (const FAttributeProto* MomentumAttribute = FModelProto::FindElementInArray(TEXT("Strides"), InNodeProto.Attribute, /*bMustValueBeFound*/false))
+	if (const FAttributeProto* MomentumAttribute = FModelProto::FindElementInArray(TEXT("Strides"), InNodeProto->Attribute, /*bMustValueBeFound*/false))
 	{
 		Strides = MomentumAttribute->Integers;
 	}
@@ -78,7 +85,7 @@ FConvOperator FPrivateConvOperator::NodeProtoToConvOperator(const FNodeProto& In
 /* FConvOperator structors
  *****************************************************************************/
 
-FConvOperator::FConvOperator(const FNodeProto& InNodeProto)
+FConvOperator::FConvOperator(const FNodeProto* const InNodeProto)
 	: FConvOperator(FPrivateConvOperator::NodeProtoToConvOperator(InNodeProto))
 {
 }
diff --git a/Engine/Plugins/Experimental/NeuralNetworkInference/Source/NeuralNetworkInference/Private/NeuralOperators/ConvTransposeOperator.cpp b/Engine/Plugins/Experimental/NeuralNetworkInference/Source/NeuralNetworkInference/Private/NeuralOperators/ConvTransposeOperator.cpp
--- a/Engine/Plugins/Experimental/NeuralNetworkInference/Source/NeuralNetworkInference/Private/NeuralOperators/ConvTransposeOperator.cpp
+++ b/Engine/Plugins/Experimental/NeuralNetworkInference/Source/NeuralNetworkInference/Private/NeuralOperators/ConvTransposeOperator.cpp
@@ -1,6 +1,7 @@
 // Copyright Epic Games, Inc. All Rights Reserved.
 
 #include "NeuralOperators/ConvTransposeOperator.h"
+#include "ModelProto.h"
 #include "NeuralNetworkInferenceShaders/ConvTransposeCS.h"
 #include "NeuralNetworkInferenceUtils.h"
 #include "NeuralNetworkInferenceUtilsGPU.h"
@@ -13,15 +14,21 @@
 class FPrivateConvTransposeOperator
 {
 public:
-	static FConvTransposeOperator NodeProtoToConvTransposeOperator(const FNodeProto& InNodeProto);
+	static FConvTransposeOperator NodeProtoToConvTransposeOperator(const FNodeProto* const InNodeProto);
 
 	static bool ZerosFromConvStrides(FNeuralInt64ArrayUInt32Buffer& OutZeros, const FNeuralInt64ArrayUInt32Buffer& InStrides, const int32 InNumberConvolutionalDimensions);
 };
 
-FConvTransposeOperator FPrivateConvTransposeOperator::NodeProtoToConvTransposeOperator(const FNodeProto& InNodeProto)
+FConvTransposeOperator FPrivateConvTransposeOperator::NodeProtoToConvTransposeOperator(const FNodeProto* const InNodeProto)
 {
+	// Sanity check
+	if (!InNodeProto)
+	{
+		UE_LOG(LogNeuralNetworkInference, Warning, TEXT("FConvTransposeOperator(): InNodeProto was a nullptr."));
+		return FConvTransposeOperator(FConvBaseOperator::EAutoPad::NotSet, {}, -1, {}, {}, {}, {}, {});
+	}
 	FConvBaseOperator::EAutoPad AutoPad = FConvBaseOperator::EAutoPad::NotSet;
-	if (const FAttributeProto* EpsilonAttribute = FModelProto::FindElementInArray(TEXT("AutoPad"), InNodeProto.Attribute, /*bMustValueBeFound*/false))
+	if (const FAttributeProto* EpsilonAttribute = FModelProto::FindElementInArray(TEXT("AutoPad"), InNodeProto->Attribute, /*bMustValueBeFound*/false))
 	{
 		if (EpsilonAttribute->S == TEXT("NOTSET"))
 		{
@@ -45,12 +52,12 @@ FConvTransposeOperator FPrivateConvTransposeOperator::NodeProtoToConvTransposeOp
 		}
 	}
 	TArray<int64> Dilations;
-	if (const FAttributeProto* MomentumAttribute = FModelProto::FindElementInArray(TEXT("Dilations"), InNodeProto.Attribute, /*bMustValueBeFound*/false))
+	if (const FAttributeProto* MomentumAttribute = FModelProto::FindElementInArray(TEXT("Dilations"), InNodeProto->Attribute, /*bMustValueBeFound*/false))
 	{
 		Dilations = MomentumAttribute->Integers;
 	}
 	int64 Group;
-	if (const FAttributeProto* MomentumAttribute = FModelProto::FindElementInArray(TEXT("Group"), InNodeProto.Attribute, /*bMustValueBeFound*/false))
+	if (const FAttributeProto* MomentumAttribute = FModelProto::FindElementInArray(TEXT("Group"), InNodeProto->Attribute, /*bMustValueBeFound*/false))
 	{
 		Group = MomentumAttribute->I;
 	}
@@ -59,32 +66,32 @@ FConvTransposeOperator FPrivateConvTransposeOperator::NodeProtoToConvTransposeOp
 		Group = 1;
 	}
 	TArray<int64> KernelShape;
-	if (const FAttributeProto* MomentumAttribute = FModelProto::FindElementInArray(TEXT("KernelShape"), InNodeProto.Attribute, /*bMustValueBeFound*/false))
+	if (const FAttributeProto* MomentumAttribute = FModelProto::FindElementInArray(TEXT("KernelShape"), InNodeProto->Attribute, /*bMustValueBeFound*/false))
 	{
 		KernelShape = MomentumAttribute->Integers;
 	}
 	TArray<int64> Pads;
-	if (const FAttributeProto* MomentumAttribute = FModelProto::FindElementInArray(TEXT("Pads"), InNodeProto.Attribute, /*bMustValueBeFound*/false))
+	if (const FAttributeProto* MomentumAttribute = FModelProto::FindElementInArray(TEXT("Pads"), InNodeProto->Attribute, /*bMustValueBeFound*/false))
 	{
 		Pads = MomentumAttribute->Integers;
 	}
 	TArray<int64> Strides;
-	if (const FAttributeProto* MomentumAttribute = FModelProto::FindElementInArray(TEXT("Strides"), InNodeProto.Attribute, /*bMustValueBeFound*/false))
+	if (const FAttributeProto* MomentumAttribute = FModelProto::FindElementInArray(TEXT("Strides"), InNodeProto->Attribute, /*bMustValueBeFound*/false))
 	{
 		Strides = MomentumAttribute->Integers;
 	}
 	TArray<int64> Zeros;
-	if (const FAttributeProto* MomentumAttribute = FModelProto::FindElementInArray(TEXT("Zeros"), InNodeProto.Attribute, /*bMustValueBeFound*/false))
+	if (const FAttributeProto* MomentumAttribute = FModelProto::FindElementInArray(TEXT("Zeros"), InNodeProto->Attribute, /*bMustValueBeFound*/false))
 	{
 		Zeros = MomentumAttribute->Integers;
 	}
 	TArray<int64> OutputPadding;
-	if (const FAttributeProto* MomentumAttribute = FModelProto::FindElementInArray(TEXT("OutputPadding"), InNodeProto.Attribute, /*bMustValueBeFound*/false))
+	if (const FAttributeProto* MomentumAttribute = FModelProto::FindElementInArray(TEXT("OutputPadding"), InNodeProto->Attribute, /*bMustValueBeFound*/false))
 	{
 		OutputPadding = MomentumAttribute->Integers;
 	}
 	TArray<int64> OutputShape;
-	if (const FAttributeProto* MomentumAttribute = FModelProto::FindElementInArray(TEXT("OutputShape"), InNodeProto.Attribute, /*bMustValueBeFound*/false))
+	if (const FAttributeProto* MomentumAttribute = FModelProto::FindElementInArray(TEXT("OutputShape"), InNodeProto->Attribute, /*bMustValueBeFound*/false))
 	{
 		OutputShape = MomentumAttribute->Integers;
 	}
@@ -117,7 +124,7 @@ bool FPrivateConvTransposeOperator::ZerosFromConvStrides(FNeuralInt64ArrayUInt32
 /* FConvTransposeOperator structors
  *****************************************************************************/
 
-FConvTransposeOperator::FConvTransposeOperator(const FNodeProto& InNodeProto)
+FConvTransposeOperator::FConvTransposeOperator(const FNodeProto* const InNodeProto)
 	: FConvTransposeOperator(FPrivateConvTransposeOperator::NodeProtoToConvTransposeOperator(InNodeProto))
 {
 }
diff --git a/Engine/Plugins/Experimental/NeuralNetworkInference/Source/NeuralNetworkInference/Private/NeuralOperators/GemmOperator.cpp b/Engine/Plugins/Experimental/NeuralNetworkInference/Source/NeuralNetworkInference/Private/NeuralOperators/GemmOperator.cpp
--- a/Engine/Plugins/Experimental/NeuralNetworkInference/Source/NeuralNetworkInference/Private/NeuralOperators/GemmOperator.cpp
+++ b/Engine/Plugins/Experimental/NeuralNetworkInference/Source/NeuralNetworkInference/Private/NeuralOperators/GemmOperator.cpp
@@ -1,6 +1,7 @@
 // Copyright Epic Games, Inc. All Rights Reserved.
 
 #include "NeuralOperators/GemmOperator.h"
+#include "ModelProto.h"
 #include "NeuralNetworkInferenceShaders/GemmCS.h"
 #include "NeuralNetworkInferenceUtils.h"
 #include "NeuralNetworkInferenceUtilsGPU.h"
@@ -10,22 +11,28 @@
 /* FGemmOperator structors
  *****************************************************************************/
 
-FGemmOperator::FGemmOperator(const FNodeProto& InNodeProto)
+FGemmOperator::FGemmOperator(const FNodeProto* const InNodeProto)
 	: FGemmOperator()
 {
-	if (const FAttributeProto* AlphaAttribute = FModelProto::FindElementInArray(TEXT("Alpha"), InNodeProto.Attribute, /*bMustValueBeFound*/false))
+	// Sanity check
+	if (!InNodeProto)
+	{
+		UE_LOG(LogNeuralNetworkInference, Warning, TEXT("FGemmOperator(): InNodeProto was a nullptr."));
+		return;
+	}
+	if (const FAttributeProto* AlphaAttribute = FModelProto::FindElementInArray(TEXT("Alpha"), InNodeProto->Attribute, /*bMustValueBeFound*/false))
 	{
 		Alpha = AlphaAttribute->F;
 	}
-	if (const FAttributeProto* BetaAttribute = FModelProto::FindElementInArray(TEXT("Beta"), InNodeProto.Attribute, /*bMustValueBeFound*/false))
+	if (const FAttributeProto* BetaAttribute = FModelProto::FindElementInArray(TEXT("Beta"), InNodeProto->Attribute, /*bMustValueBeFound*/false))
 	{
 		Beta = BetaAttribute->F;
 	}
-	if (const FAttributeProto* TransAAttribute = FModelProto::FindElementInArray(TEXT("TransA"), InNodeProto.Attribute, /*bMustValueBeFound*/false))
+	if (const FAttributeProto* TransAAttribute = FModelProto::FindElementInArray(TEXT("TransA"), InNodeProto->Attribute, /*bMustValueBeFound*/false))
 	{
 		bTransA = (TransAAttribute->I != 0);
 	}
-	if (const FAttributeProto* TransBAttribute = FModelProto::FindElementInArray(TEXT("TransB"), InNodeProto.Attribute, /*bMustValueBeFound*/false))
+	if (const FAttributeProto* TransBAttribute = FModelProto::FindElementInArray(TEXT("TransB"), InNodeProto->Attribute, /*bMustValueBeFound*/false))
 	{
 		bTransB = (TransBAttribute->I != 0);
 	}
diff --git a/Engine/Plugins/Experimental/NeuralNetworkInference/Source/NeuralNetworkInference/Private/NeuralOperators/LeakyReluOperator.cpp b/Engine/Plugins/Experimental/NeuralNetworkInference/Source/NeuralNetworkInference/Private/NeuralOperators/LeakyReluOperator.cpp
--- a/Engine/Plugins/Experimental/NeuralNetworkInference/Source/NeuralNetworkInference/Private/NeuralOperators/LeakyReluOperator.cpp
+++ b/Engine/Plugins/Experimental/NeuralNetworkInference/Source/NeuralNetworkInference/Private/NeuralOperators/LeakyReluOperator.cpp
@@ -1,6 +1,7 @@
 // Copyright Epic Games, Inc. All Rights Reserved.
 
 #include "NeuralOperators/LeakyReluOperator.h"
+#include "ModelProto.h"
 #include "NeuralNetworkInferenceUtils.h"
 
 
@@ -8,11 +9,17 @@
 /* FLeakyReluOperator structors
  *****************************************************************************/
 
-FLeakyReluOperator::FLeakyReluOperator(const bool bIsInlinedTensor, const FNodeProto& InNodeProto)
+FLeakyReluOperator::FLeakyReluOperator(const bool bIsInlinedTensor, const FNodeProto* const InNodeProto)
 	: FLeakyReluOperator(bIsInlinedTensor)
 {
 	UE_LOG(LogNeuralNetworkInference, Warning, TEXT("LeakyRelu(): Constructor not tested yet."));
-	if (const FAttributeProto* AlphaAttribute = FModelProto::FindElementInArray(TEXT("Alpha"), InNodeProto.Attribute, /*bMustValueBeFound*/false))
+	// Sanity check
+	if (!InNodeProto)
+	{
+		UE_LOG(LogNeuralNetworkInference, Warning, TEXT("FLeakyReluOperator(): InNodeProto was a nullptr."));
+		return;
+	}
+	if (const FAttributeProto* AlphaAttribute = FModelProto::FindElementInArray(TEXT("Alpha"), InNodeProto->Attribute, /*bMustValueBeFound*/false))
 	{
 		Attributes[0] = AlphaAttribute->F;
 	}
diff --git a/Engine/Plugins/Experimental/NeuralNetworkInference/Source/NeuralNetworkInference/Private/NeuralOperators/ReshapeOperator.cpp b/Engine/Plugins/Experimental/NeuralNetworkInference/Source/NeuralNetworkInference/Private/NeuralOperators/ReshapeOperator.cpp
--- a/Engine/Plugins/Experimental/NeuralNetworkInference/Source/NeuralNetworkInference/Private/NeuralOperators/ReshapeOperator.cpp
+++ b/Engine/Plugins/Experimental/NeuralNetworkInference/Source/NeuralNetworkInference/Private/NeuralOperators/ReshapeOperator.cpp
@@ -1,6 +1,7 @@
 // Copyright Epic Games, Inc. All Rights Reserved.
 
 #include "NeuralOperators/ReshapeOperator.h"
+#include "ModelProto.h"
 #include "NeuralNetworkInferenceShaders/CopyCS.h"
 #include "NeuralNetworkInferenceUtils.h"
 #include "NeuralNetworkInferenceUtilsGPU.h"
@@ -10,10 +11,16 @@
 /* FReshapeOperator structors
  *****************************************************************************/
 
-FReshapeOperator::FReshapeOperator(const bool bIsInlinedTensor, const FNodeProto& InNodeProto)
+FReshapeOperator::FReshapeOperator(const bool bIsInlinedTensor, const FNodeProto* const InNodeProto)
 	: FReshapeOperator(bIsInlinedTensor)
 {
-	if (const FAttributeProto* MomentumAttribute = FModelProto::FindElementInArray(TEXT("AllowZero"), InNodeProto.Attribute, /*bMustValueBeFound*/false))
+	// Sanity check
+	if (!InNodeProto)
+	{
+		UE_LOG(LogNeuralNetworkInference, Warning, TEXT("FReshapeOperator(): InNodeProto was a nullptr."));
+		return;
+	}
+	if (const FAttributeProto* MomentumAttribute = FModelProto::FindElementInArray(TEXT("AllowZero"), InNodeProto->Attribute, /*bMustValueBeFound*/false))
 	{
 		AllowZero = MomentumAttribute->I;
 	}
diff --git a/Engine/Plugins/Experimental/NeuralNetworkInference/Source/NeuralNetworkInference/Private/UEOnly/GraphProtoToNeuralNetworkConverter.cpp b/Engine/Plugins/Experimental/NeuralNetworkInference/Source/NeuralNetworkInference/Private/UEOnly/GraphProtoToNeuralNetworkConverter.cpp
--- a/Engine/Plugins/Experimental/NeuralNetworkInference/Source/NeuralNetworkInference/Private/UEOnly/GraphProtoToNeuralNetworkConverter.cpp
+++ b/Engine/Plugins/Experimental/NeuralNetworkInference/Source/NeuralNetworkInference/Private/UEOnly/GraphProtoToNeuralNetworkConverter.cpp
@@ -227,7 +227,7 @@ bool FGraphProtoToNeuralNetworkConverter::CreateOperatorsAndEditTensorArray(TArr
 				InOutTensors.Push(FNeuralTensor(ENeuralDataType::None, TArray<int64>({}), OutputTensorName, ENeuralTensorTypeGPU::Weight));
 
 				// Create constant operator layer and get final tensor value
-				FConstantOperator ConstantOperator(NodeProto);
+				FConstantOperator ConstantOperator(&NodeProto);
 				ConstantOperator.SetOutputTensors({&InOutTensors.Last()});
 				// Sanity check
 				if (!ConstantOperator.ConfigureOutputAndInternalVariablesAndSanityChecks())
@@ -539,14 +539,14 @@ TSet<uint32> FGraphProtoToNeuralNetworkConverter::GetPotentiallyInlinedTensors(c
 	if (InNodeProto.OperatorType == TEXT(#OperatorTypeSuffix)) /* Eg TEXT("BatchNormalization") */ \
 	{ \
 		const bool bCanTensorBeInlined = CanTensorBeInlined(InInputTensorNamesForOperator[0], InInputNameDummyIndexMap, InOutputNameDummyIndexMap, InNodeProto, InGraphProto, bCanOutputPropertiesChangeWithPostForward); \
-		Operator = MakeShared<F##OperatorTypeSuffix##Operator>(bCanTensorBeInlined, InNodeProto); /* Eg FBatchNormalizationOperator */ \
+		Operator = MakeShared<F##OperatorTypeSuffix##Operator>(bCanTensorBeInlined, &InNodeProto); /* Eg FBatchNormalizationOperator */ \
 	}
 
 // Eg Conv(Transpose) and Gemm
 #define IF_MATCHED_OPERATOR_WITH_PROTO_INITIALIZATION(OperatorTypeSuffix) \
 	if (InNodeProto.OperatorType == TEXT(#OperatorTypeSuffix)) /* Eg TEXT("Conv") */ \
 	{ \
-		Operator = MakeShared<F##OperatorTypeSuffix##Operator>(InNodeProto); /* Eg FConvOperator */ \
+		Operator = MakeShared<F##OperatorTypeSuffix##Operator>(&InNodeProto); /* Eg FConvOperator */ \
 	}
 
 #define IF_MATCHED_MULTIDIRECTIONAL_BROADCASTING_OPERATOR_INITIALIZATION(OperatorTypeSuffix) \
@@ -572,7 +572,7 @@ TSharedPtr<FNeuralOperator> FGraphProtoToNeuralNetworkConverter::CreateOperator(
 	if (InNodeProto.OperatorType == TEXT("ConvTranspose"))
 	{
 		TSharedPtr<FConvTransposeOperator> ConvTransposeOperator;
-		ConvTransposeOperator = MakeShared<FConvTransposeOperator>(InNodeProto);
+		ConvTransposeOperator = MakeShared<FConvTransposeOperator>(&InNodeProto);
 		ConvTransposeOperator->SetWhetherWeightTensorWasFlipped(true);
 		Operator = MoveTemp(ConvTransposeOperator);
 	}
