NeuralStats.h:
1. Hugely simplified to improve client readability.
2. Added "const" to some functions that didn't have it (also the stats ones from UNN).
3. FNNIStatsData renamed as FNeuralStatsData to be more consistent with FNeuralStats

#rnx

diff --git a/Engine/Plugins/Experimental/NeuralNetworkInference/Source/NeuralNetworkInference/Private/NeuralNetwork.cpp b/Engine/Plugins/Experimental/NeuralNetworkInference/Source/NeuralNetworkInference/Private/NeuralNetwork.cpp
--- a/Engine/Plugins/Experimental/NeuralNetworkInference/Source/NeuralNetworkInference/Private/NeuralNetwork.cpp
+++ b/Engine/Plugins/Experimental/NeuralNetworkInference/Source/NeuralNetworkInference/Private/NeuralNetwork.cpp
@@ -516,15 +516,18 @@ void UNeuralNetwork::Run()
 	ComputeStatsModule.StoreSample(RunTimer.Toc());
 }
 
-float UNeuralNetwork::GetLastInferenceTime() {
+float UNeuralNetwork::GetLastInferenceTime() const
+{
 	return ComputeStatsModule.GetLastSample();
 };
 
-FNNIStatsData UNeuralNetwork::GetInferenceStats() {
+FNeuralStatsData UNeuralNetwork::GetInferenceStats() const
+{
 	return ComputeStatsModule.GetStats();
 };
 
-FNNIStatsData UNeuralNetwork::GetInputMemoryTransferStats() {
+FNeuralStatsData UNeuralNetwork::GetInputMemoryTransferStats() const
+{
 	return InputMemoryTransferStatsModule.GetStats();
 };
 
diff --git a/Engine/Plugins/Experimental/NeuralNetworkInference/Source/NeuralNetworkInference/Public/NeuralNetwork.h b/Engine/Plugins/Experimental/NeuralNetworkInference/Source/NeuralNetworkInference/Public/NeuralNetwork.h
--- a/Engine/Plugins/Experimental/NeuralNetworkInference/Source/NeuralNetworkInference/Public/NeuralNetwork.h
+++ b/Engine/Plugins/Experimental/NeuralNetworkInference/Source/NeuralNetworkInference/Public/NeuralNetwork.h
@@ -182,9 +182,9 @@ public:
 	 * - GetInferenceStats, returns Inference time statistics. (NumberSamples, Average, StdDev, Min, Max statistics measured in milliseconds)
 	 * - GetInputMemoryTransferStats, returns Input Memory Transfer statistics. (NumberSamples, Average, StdDev, Min, Max statistics measured in milliseconds)
 	 */
-	float GetLastInferenceTime();
-	FNNIStatsData GetInferenceStats();
-	FNNIStatsData GetInputMemoryTransferStats();
+	float GetLastInferenceTime() const;
+	FNeuralStatsData GetInferenceStats() const;
+	FNeuralStatsData GetInputMemoryTransferStats() const;
 	void ResetStats();
 
 protected:
diff --git a/Engine/Plugins/Experimental/NeuralNetworkInference/Source/NeuralNetworkInferenceProfiling/Private/NeuralStats.cpp b/Engine/Plugins/Experimental/NeuralNetworkInference/Source/NeuralNetworkInferenceProfiling/Private/NeuralStats.cpp
--- a/Engine/Plugins/Experimental/NeuralNetworkInference/Source/NeuralNetworkInferenceProfiling/Private/NeuralStats.cpp
+++ b/Engine/Plugins/Experimental/NeuralNetworkInference/Source/NeuralNetworkInferenceProfiling/Private/NeuralStats.cpp
@@ -2,15 +2,59 @@
 
 #include "NeuralStats.h"
 #include "GenericPlatform/GenericPlatformMath.h"
+
 DEFINE_LOG_CATEGORY_STATIC(LogNNIStats, Display, All);
 
 
+
+/* FNeuralStats auxiliary functions
+ *****************************************************************************/
+
+void FNeuralStats_CalculateMinMax(float& OutMin, float& OutMax, const TCircularBuffer<float>& InDataBuffer, const bool bInIsBufferFull, const uint32 InBufferIdx)
+{
+	const uint32 PtrEnd = (bInIsBufferFull) ? InDataBuffer.Capacity() : InBufferIdx;
+
+	OutMin = FLT_MAX;
+	OutMax = FLT_MIN;
+	for (uint32 Index = 0; Index < PtrEnd; ++Index)
+	{
+		const float CurrentElement = InDataBuffer[Index];
+		if (CurrentElement < OutMin)
+		{
+			OutMin = CurrentElement;
+		}
+		if (CurrentElement > OutMax)
+		{
+			OutMax = CurrentElement;
+		}
+	}
+}
+
+
+
+/* FNeuralStatsData auxiliary functions
+ *****************************************************************************/
+
+FNeuralStatsData::FNeuralStatsData(const int InNumberSamples, const float InAverage, const float InStdDev, const float InMin, const float InMax)
+	: NumberSamples(InNumberSamples)
+	, Average(InAverage)
+	, StdDev(InStdDev)
+	, Min(InMin)
+	, Max(InMax)
+{}
+
+
+
+/* FNeuralNetworkInferenceQATimer public functions
+ *****************************************************************************/
+
 FNeuralStats::FNeuralStats(const uint32 InSizeRollingWindow) 
-: DataBuffer(InSizeRollingWindow, 0)
+	: EpsilonFloat(0.000001f)
+	, BufferIdx(0)
+	, bIsBufferFull(false)
+	, LastSample(0.f)
+	, DataBuffer(InSizeRollingWindow, 0)
 {
-	SizeRollingWindow = InSizeRollingWindow;
-	BufferIdx = 0;
-	bIsBufferFull = false;
 }
 	
 void FNeuralStats::StoreSample(const float InRunTime) 
@@ -24,91 +68,56 @@ void FNeuralStats::StoreSample(const float InRunTime)
 		bIsBufferFull = true;
 	}
 	BufferIdx = DataBuffer.GetNextIndex(BufferIdx);
+}
 
+void FNeuralStats::ResetStats()
+{
+	for (uint32 Index = 0; Index < DataBuffer.Capacity(); ++Index)
+	{
+		DataBuffer[Index] = 0;
+	}
+	bIsBufferFull = false;
+	BufferIdx = 0;
+	LastSample = 0;
 }
 
-float FNeuralStats::GetLastSample()
+float FNeuralStats::GetLastSample() const
 {
 	return LastSample;
 }
 
-FNNIStatsData FNeuralStats::GetStats()
+FNeuralStatsData FNeuralStats::GetStats() const
 {
 	const uint32 count = (bIsBufferFull) ? DataBuffer.Capacity() : BufferIdx;
 	const float Mean = CalculateMean();
 	const float StdDev = CalculateStdDev(Mean);
-	FNNIStatsMinMax MinMaxData = CalculateMinMax();
-
-	FNNIStatsData AggregatedData(count,	Mean, StdDev, MinMaxData.Min, MinMaxData.Max);
-
-	return AggregatedData;
+	float Min, Max;
+	FNeuralStats_CalculateMinMax(Min, Max, DataBuffer, bIsBufferFull, BufferIdx);
+	// Return aggregated data
+	return FNeuralStatsData(count, Mean, StdDev, Min, Max);
 }
 
-void FNeuralStats::SetSizeRollingWindow(const uint32 InSizeRollingWindow) 
-{
-	DataBuffer = TCircularBuffer<float>(InSizeRollingWindow, 0);
-	ResetStats();
-}
-
-float FNeuralStats::CalculateMean() 
+float FNeuralStats::CalculateMean() const
 {
 	uint32 PtrEnd = (bIsBufferFull) ? DataBuffer.Capacity(): BufferIdx;
 	float AccumulatorMean = 0;
-
 	for (uint32 Index = 0; Index < PtrEnd; ++Index)
 	{
 		AccumulatorMean += DataBuffer[Index];
 	}
-
-	const float Mean = AccumulatorMean / (static_cast<float>(PtrEnd) + EpsilonFloat);
-
-	return Mean;
+	// Return mean
+	return AccumulatorMean / (static_cast<float>(PtrEnd) + EpsilonFloat);
 }
 
-float FNeuralStats::CalculateStdDev(const float InMean)
+float FNeuralStats::CalculateStdDev(const float InMean) const
 {
-	uint32 PtrEnd = (bIsBufferFull) ? DataBuffer.Capacity() : BufferIdx;
+	const uint32 PtrEnd = (bIsBufferFull) ? DataBuffer.Capacity() : BufferIdx;
 
 	float AccumulatorStdDev = 0;
-
 	for (uint32 Index = 0; Index < PtrEnd; ++Index)
 	{
 		AccumulatorStdDev += (DataBuffer[Index] - InMean) * (DataBuffer[Index] - InMean);
 	}
-
-	const float StdDev = FMath::Sqrt(AccumulatorStdDev / (static_cast<float>(PtrEnd) + EpsilonFloat));
-
-	return StdDev;
-}
-
-FNNIStatsMinMax FNeuralStats::CalculateMinMax()
-{
-	FNNIStatsMinMax MinMaxData;
-	uint32 PtrEnd = (bIsBufferFull) ? DataBuffer.Capacity() : BufferIdx;
-
-	for (uint32 Index = 0; Index < PtrEnd; ++Index)
-	{
-		const float CurrentElement = DataBuffer[Index];
-		if( CurrentElement < MinMaxData.Min)
-		{
-			MinMaxData.Min = CurrentElement;
-		}
-		if (CurrentElement > MinMaxData.Max)
-		{
-			MinMaxData.Max = CurrentElement;
-		}
-	}
-
-	return MinMaxData;
-}
-
-void FNeuralStats::ResetStats()
-{
-	for (uint32 Index = 0; Index < DataBuffer.Capacity(); ++Index)
-	{
-		DataBuffer[Index] = 0;
-	}
-	bIsBufferFull = false;
-	BufferIdx = 0;
-	LastSample = 0;
+	// Return std deviation
+	return FMath::Sqrt(AccumulatorStdDev / (static_cast<float>(PtrEnd) + EpsilonFloat));
 }
diff --git a/Engine/Plugins/Experimental/NeuralNetworkInference/Source/NeuralNetworkInferenceProfiling/Public/NeuralStats.h b/Engine/Plugins/Experimental/NeuralNetworkInference/Source/NeuralNetworkInferenceProfiling/Public/NeuralStats.h
--- a/Engine/Plugins/Experimental/NeuralNetworkInference/Source/NeuralNetworkInferenceProfiling/Public/NeuralStats.h
+++ b/Engine/Plugins/Experimental/NeuralNetworkInference/Source/NeuralNetworkInferenceProfiling/Public/NeuralStats.h
@@ -5,71 +5,35 @@
 #include "CoreMinimal.h"
 #include "Containers/CircularBuffer.h"
 
-
-struct NEURALNETWORKINFERENCEPROFILING_API FNNISampleStatData
+struct NEURALNETWORKINFERENCEPROFILING_API FNeuralStatsData
 {
-	FNNISampleStatData() {};
-
-	FNNISampleStatData(const float InInferenceTime)
-		: InferenceTime(InInferenceTime)	
-	{}
+	const uint32 NumberSamples;
+	const float Average;
+	const float StdDev;
+	const float Min;
+	const float Max;
 
-	float InferenceTime = 0;
+	FNeuralStatsData(const int InNumberSamples, const float InAverage, const float InStdDev, const float InMin, const float InMax);
 };
 
-struct NEURALNETWORKINFERENCEPROFILING_API FNNIStatsData
-{
-	FNNIStatsData() {}
-	
-	FNNIStatsData(const int InNumberSamples, const float InAverage, const float InStdDev, const float InMin, const float InMax)
-		: NumberSamples(InNumberSamples)
-		, Average(InAverage)
-		, StdDev(InStdDev)
-		, Min(InMin)
-		, Max(InMax)
-		{}
-
-	uint32 NumberSamples = 0;
-	float Average = 0;
-	float StdDev = 0;
-	float Min = 0;
-	float Max = 0;
-};
-
-struct NEURALNETWORKINFERENCEPROFILING_API FNNIStatsMinMax
-{
-	float Min = FLT_MAX;
-	float Max = FLT_MIN;
-};
-
-
 class NEURALNETWORKINFERENCEPROFILING_API FNeuralStats
 {
-
 public:
-
 	FNeuralStats(const uint32 InSizeRollingWindow = 1024);
 	
 	void StoreSample(const float InRunTime);
 	void ResetStats();
-	void SetSizeRollingWindow(const uint32 InSizeRollingWindow);
 
-	float GetLastSample();
-	FNNIStatsData GetStats();
+	float GetLastSample() const;
+	FNeuralStatsData GetStats() const;
 
 private:
-
-	uint32 BufferIdx = 0;
-	bool bIsBufferFull = false;
-	uint32 SizeRollingWindow = 1024;
-	const float EpsilonFloat = 0.000001;
+	const float EpsilonFloat;
+	uint32 BufferIdx;
+	bool bIsBufferFull;
 	float LastSample;
 	TCircularBuffer<float> DataBuffer;
 
-	float CalculateMean();
-	float CalculateStdDev(const float InMean);
-	FNNIStatsMinMax CalculateMinMax();
-	
+	float CalculateMean() const;
+	float CalculateStdDev(const float InMean) const;
 };
-
-
diff --git a/Engine/Plugins/Experimental/NeuralNetworkInference/Source/NeuralNetworkInferenceQA/Private/ModelUnitTester.cpp b/Engine/Plugins/Experimental/NeuralNetworkInference/Source/NeuralNetworkInferenceQA/Private/ModelUnitTester.cpp
--- a/Engine/Plugins/Experimental/NeuralNetworkInference/Source/NeuralNetworkInferenceQA/Private/ModelUnitTester.cpp
+++ b/Engine/Plugins/Experimental/NeuralNetworkInference/Source/NeuralNetworkInferenceQA/Private/ModelUnitTester.cpp
@@ -6,13 +6,61 @@
 #include "Misc/Paths.h"
 
 
+
+/* FModelUnitTester axuiliary functions
+ *****************************************************************************/
+
 struct FNNIUnitTesterTimeData 
 {
-	FNNIStatsData ComputeTimeData;
-	FNNIStatsData InputCopyTimeData;
-	FNNIStatsData OutputCopyTimeData;
+	FNeuralStatsData ComputeTimeData;
+	FNeuralStatsData InputCopyTimeData;
+	FNeuralStatsData OutputCopyTimeData;
+
+	FNNIUnitTesterTimeData(const FNeuralStatsData& InComputeTimeData, const FNeuralStatsData& InInputCopyTimeData, const FNeuralStatsData& InOutputCopyTimeData)
+		: ComputeTimeData(InComputeTimeData)
+		, InputCopyTimeData(InInputCopyTimeData)
+		, OutputCopyTimeData(InOutputCopyTimeData)
+	{}
 };
 
+FNNIUnitTesterTimeData FModelUnitTester_GetTimeInformation(UNeuralNetwork* InOutNetwork, TArray<float>& OutCPUGPUCPUOutput, const TArray<float>& InInputArray, const int32 InRepetitions) 
+{
+	/* Input/output copy speed */ 
+	FNeuralStats OutCopyingStats;
+	FNeuralNetworkInferenceTimer Timer;
+	for (int32 TimerIndex = 0; TimerIndex < InRepetitions; ++TimerIndex) 
+	{ 
+		InOutNetwork->SetInputFromArrayCopy(InInputArray); 
+		Timer.Tic();
+		OutCPUGPUCPUOutput = InOutNetwork->GetOutputTensor().GetArrayCopy<float>();
+		const float CurrentCopyingTime = Timer.Toc();
+		OutCopyingStats.StoreSample(CurrentCopyingTime);
+
+	} 
+
+	/* Forward() speed */ 
+	if (InRepetitions > 1) 
+	{ 
+		for (int32 TimerIndex = 0; TimerIndex < 5; ++TimerIndex) 
+		{ 
+			InOutNetwork->Run(); 
+		} 
+	} 
+
+	if (InRepetitions > 0) 
+	{ 
+		for (int32 TimerIndex = 0; TimerIndex < InRepetitions; ++TimerIndex) 
+		{ 
+			InOutNetwork->SetInputFromArrayCopy(InInputArray);
+			InOutNetwork->Run();
+			OutCPUGPUCPUOutput = InOutNetwork->GetOutputTensor().GetArrayCopy<float>(); 
+		} 
+	}
+	// Return NetworkTimeData
+	return FNNIUnitTesterTimeData(InOutNetwork->GetInferenceStats(), InOutNetwork->GetInputMemoryTransferStats(), OutCopyingStats.GetStats());
+}
+
+
 
 /* FModelUnitTester static public functions
  *****************************************************************************/
@@ -364,54 +412,6 @@ else
 	return bDidGlobalTestPassed;
 }
 
-FNNIUnitTesterTimeData GetTimeInformation(
-		const int32 InRepetitions,
-		FNeuralNetworkInferenceTimer& Timer,
-		UNeuralNetwork* InOutNetwork,
-		TArray<float>* InputArray,
-		TArray<float>* CPUGPUCPUOutput) 
-	{
-
-		FNeuralStats OutCopyingStats;
-
-		/* Input/output copy speed */ 
-		for (int32 TimerIndex = 0; TimerIndex < InRepetitions; ++TimerIndex) 
-		{ 
-			InOutNetwork->SetInputFromArrayCopy(*InputArray); 
-			Timer.Tic();
-			*CPUGPUCPUOutput = InOutNetwork->GetOutputTensor().GetArrayCopy<float>();
-			const float CurrentCopyingTime = Timer.Toc();
-			OutCopyingStats.StoreSample(CurrentCopyingTime);
-
-		} 
-
-		/* Forward() speed */ 
-		if (InRepetitions > 1) 
-		{ 
-			for (int32 TimerIndex = 0; TimerIndex < 5; ++TimerIndex) 
-			{ 
-				InOutNetwork->Run(); 
-			} 
-		} 
-
-		FNeuralStats CurrentComputingStats;
-		if (InRepetitions > 0) 
-		{ 
-			for (int32 TimerIndex = 0; TimerIndex < InRepetitions; ++TimerIndex) 
-			{ 
-				InOutNetwork->SetInputFromArrayCopy(*InputArray); 
-				InOutNetwork->Run();
-				*CPUGPUCPUOutput = InOutNetwork->GetOutputTensor().GetArrayCopy<float>(); 
-			} 
-		} 
-		FNNIUnitTesterTimeData NetworkTimeData;
-		NetworkTimeData.ComputeTimeData = InOutNetwork->GetInferenceStats();
-		NetworkTimeData.InputCopyTimeData = InOutNetwork->GetInputMemoryTransferStats();
-		NetworkTimeData.OutputCopyTimeData = OutCopyingStats.GetStats();
-
-		return NetworkTimeData;
-	}
-
 
 bool FModelUnitTester::ModelSpeedTest(const FString& InUAssetPath, const ENeuralDeviceType InDeviceType, const ENeuralBackEnd InBackEnd, const int32 InRepetitions)
 {
@@ -446,17 +446,15 @@ bool FModelUnitTester::ModelSpeedTest(const FString& InUAssetPath, const ENeural
 	TArray<float> InputArray;
 	InputArray.Init(1.f, NetworkSize);
 	TArray<float> CPUGPUCPUOutput;
-	// Speed profiling
-	FNeuralNetworkInferenceTimer Timer;
 	// Run profiling 1 time
 	InOutNetwork->SetDeviceType(InDeviceType);
 	InOutNetwork->ResetStats();
-	FNNIUnitTesterTimeData TimeData1 = GetTimeInformation(1, Timer, InOutNetwork, &InputArray, &CPUGPUCPUOutput);
+	const FNNIUnitTesterTimeData TimeData1 = FModelUnitTester_GetTimeInformation(InOutNetwork, CPUGPUCPUOutput, InputArray, 1);
 
 	// Run profiling n times
 	InOutNetwork->SetDeviceType(InDeviceType);
 	InOutNetwork->ResetStats();
-	FNNIUnitTesterTimeData TimeDataN = GetTimeInformation(InRepetitions, Timer, InOutNetwork, &InputArray, &CPUGPUCPUOutput);
+	const FNNIUnitTesterTimeData TimeDataN = FModelUnitTester_GetTimeInformation(InOutNetwork, CPUGPUCPUOutput, InputArray, InRepetitions);
 
 	// Display speed times
 	UE_LOG(LogNeuralNetworkInferenceQA, Display,
