FNNIProfiling:
- NNIModule.h/cpp --> Module.h/cpp
- FNNITimer --> FNeuralTimer

_codereview Francisco.Vicente
#rnx

diff --git a/Engine/Plugins/Experimental/NeuralNetworkInference/Source/NeuralNetworkInference/Private/NeuralNetwork.cpp b/Engine/Plugins/Experimental/NeuralNetworkInference/Source/NeuralNetworkInference/Private/NeuralNetwork.cpp
--- a/Engine/Plugins/Experimental/NeuralNetworkInference/Source/NeuralNetworkInference/Private/NeuralNetwork.cpp
+++ b/Engine/Plugins/Experimental/NeuralNetworkInference/Source/NeuralNetworkInference/Private/NeuralNetwork.cpp
@@ -5,7 +5,7 @@
 #include "NeuralNetworkImplBackEndUEOnly.h"
 #include "NeuralNetworkInferenceUtils.h"
 #include "EditorFramework/AssetImportData.h"
-#include "NeuralNetworkInferenceTimer.h"
+#include "NeuralTimer.h"
 #include "HAL/FileManager.h"
 #include "Misc/FileHelper.h"
 #include "Misc/Paths.h"
@@ -258,7 +258,7 @@ void UNeuralNetwork::SetInputFromArrayCopy(const TArray<float>& InArray, const i
 		UE_LOG(LogNeuralNetworkInference, Warning, TEXT("UNeuralNetwork::SetInputFromArrayCopy(): Call UNeuralNetwork::Load() to load a model first."));
 	}
 	
-	FNeuralNetworkInferenceTimer RunTimer;
+	FNeuralTimer RunTimer;
 	RunTimer.Tic();
 	// UEAndORT
 	if (BackEndForCurrentPlatform == ENeuralBackEnd::UEAndORT)
@@ -381,7 +381,7 @@ void UNeuralNetwork::SetInputFromArrayCopy(const TArray<FNeuralTensor>& InInputT
 		UE_LOG(LogNeuralNetworkInference, Warning, TEXT("UNeuralNetwork::SetInputFromArrayCopy(): GetInputTensorNumber() == InInputTensorArray.Num() failed, %d != %d."), GetInputTensorNumber(), InInputTensorArray.Num());
 		return;
 	}
-	FNeuralNetworkInferenceTimer RunTimer;
+	FNeuralTimer RunTimer;
 	RunTimer.Tic();
 	for (uint32 InputTensorIndex = 0; InputTensorIndex < GetInputTensorNumber(); ++InputTensorIndex)
 	{
@@ -496,7 +496,7 @@ void UNeuralNetwork::Run()
 		UE_LOG(LogNeuralNetworkInference, Warning, TEXT("UNeuralNetwork::Run(): Call UNeuralNetwork::Load() to load a model first."));
 		return;
 	}
-	FNeuralNetworkInferenceTimer RunTimer;
+	FNeuralTimer RunTimer;
 	RunTimer.Tic();
 	// UEAndORT
 	if (BackEndForCurrentPlatform == ENeuralBackEnd::UEAndORT)
diff --git a/Engine/Plugins/Experimental/NeuralNetworkInference/Source/NeuralNetworkInferenceProfiling/Private/Module.cpp b/Engine/Plugins/Experimental/NeuralNetworkInference/Source/NeuralNetworkInferenceProfiling/Private/Module.cpp
new file mode 100644
--- /dev/null
+++ b/Engine/Plugins/Experimental/NeuralNetworkInference/Source/NeuralNetworkInferenceProfiling/Private/Module.cpp
@@ -0,0 +1,16 @@
+// Copyright Epic Games, Inc. All Rights Reserved.
+
+#include "Module.h"
+
+// This code will execute after your module is loaded into memory; the exact timing is specified in the .uplugin file per-module
+void FNeuralNetworkInferenceProfilingModule::StartupModule()
+{
+}
+
+// This function may be called during shutdown to clean up your module. For modules that support dynamic reloading,
+// we call this function before unloading the module.
+void FNeuralNetworkInferenceProfilingModule::ShutdownModule()
+{
+}
+
+IMPLEMENT_MODULE(FNeuralNetworkInferenceProfilingModule, NeuralNetworkInferenceProfiling);
diff --git a/Engine/Plugins/Experimental/NeuralNetworkInference/Source/NeuralNetworkInferenceProfiling/Private/Module.h b/Engine/Plugins/Experimental/NeuralNetworkInference/Source/NeuralNetworkInferenceProfiling/Private/Module.h
new file mode 100644
--- /dev/null
+++ b/Engine/Plugins/Experimental/NeuralNetworkInference/Source/NeuralNetworkInferenceProfiling/Private/Module.h
@@ -0,0 +1,13 @@
+// Copyright Epic Games, Inc. All Rights Reserved.
+
+#pragma once
+
+#include "Modules/ModuleManager.h"
+
+class FNeuralNetworkInferenceProfilingModule : public IModuleInterface
+{
+public:
+	// Begin IModuleInterface
+	virtual void StartupModule() override;
+	virtual void ShutdownModule() override;
+};
diff --git a/Engine/Plugins/Experimental/NeuralNetworkInference/Source/NeuralNetworkInferenceProfiling/Private/NeuralNetworkInferenceProfilingModule.cpp b/Engine/Plugins/Experimental/NeuralNetworkInference/Source/NeuralNetworkInferenceProfiling/Private/NeuralNetworkInferenceProfilingModule.cpp
deleted file mode 100644
diff --git a/Engine/Plugins/Experimental/NeuralNetworkInference/Source/NeuralNetworkInferenceProfiling/Private/NeuralNetworkInferenceProfilingModule.h b/Engine/Plugins/Experimental/NeuralNetworkInference/Source/NeuralNetworkInferenceProfiling/Private/NeuralNetworkInferenceProfilingModule.h
deleted file mode 100644
diff --git a/Engine/Plugins/Experimental/NeuralNetworkInference/Source/NeuralNetworkInferenceProfiling/Private/NeuralNetworkInferenceTimer.cpp b/Engine/Plugins/Experimental/NeuralNetworkInference/Source/NeuralNetworkInferenceProfiling/Private/NeuralNetworkInferenceTimer.cpp
deleted file mode 100644
diff --git a/Engine/Plugins/Experimental/NeuralNetworkInference/Source/NeuralNetworkInferenceProfiling/Private/NeuralTimer.cpp b/Engine/Plugins/Experimental/NeuralNetworkInference/Source/NeuralNetworkInferenceProfiling/Private/NeuralTimer.cpp
new file mode 100644
--- /dev/null
+++ b/Engine/Plugins/Experimental/NeuralNetworkInference/Source/NeuralNetworkInferenceProfiling/Private/NeuralTimer.cpp
@@ -0,0 +1,17 @@
+// Copyright Epic Games, Inc. All Rights Reserved.
+
+#include "NeuralTimer.h"
+#include "HAL/PlatformTime.h"
+
+/* FNeuralNetworkInferenceQATimer public functions
+ *****************************************************************************/
+
+void FNeuralTimer::Tic()
+{
+	TimeStart = {FPlatformTime::Seconds()};
+}
+
+double FNeuralTimer::Toc() const
+{
+	return (FPlatformTime::Seconds() - TimeStart) * 1e3;
+}
diff --git a/Engine/Plugins/Experimental/NeuralNetworkInference/Source/NeuralNetworkInferenceProfiling/Public/NeuralNetworkInferenceTimer.h b/Engine/Plugins/Experimental/NeuralNetworkInference/Source/NeuralNetworkInferenceProfiling/Public/NeuralNetworkInferenceTimer.h
deleted file mode 100644
diff --git a/Engine/Plugins/Experimental/NeuralNetworkInference/Source/NeuralNetworkInferenceProfiling/Public/NeuralTimer.h b/Engine/Plugins/Experimental/NeuralNetworkInference/Source/NeuralNetworkInferenceProfiling/Public/NeuralTimer.h
new file mode 100644
--- /dev/null
+++ b/Engine/Plugins/Experimental/NeuralNetworkInference/Source/NeuralNetworkInferenceProfiling/Public/NeuralTimer.h
@@ -0,0 +1,19 @@
+// Copyright Epic Games, Inc. All Rights Reserved.
+
+#pragma once
+
+#include "CoreMinimal.h"
+
+class NEURALNETWORKINFERENCEPROFILING_API FNeuralTimer
+{
+public:
+	void Tic();
+
+	/**
+	 * Time in milliseconds, but with nanosecond accuracy (e.g., 1.234567 msec).
+	 */
+	double Toc() const;
+
+private:
+	double TimeStart;
+};
diff --git a/Engine/Plugins/Experimental/NeuralNetworkInference/Source/NeuralNetworkInferenceQA/Private/ModelUnitTester.cpp b/Engine/Plugins/Experimental/NeuralNetworkInference/Source/NeuralNetworkInferenceQA/Private/ModelUnitTester.cpp
--- a/Engine/Plugins/Experimental/NeuralNetworkInference/Source/NeuralNetworkInferenceQA/Private/ModelUnitTester.cpp
+++ b/Engine/Plugins/Experimental/NeuralNetworkInference/Source/NeuralNetworkInferenceQA/Private/ModelUnitTester.cpp
@@ -1,7 +1,7 @@
 // Copyright Epic Games, Inc. All Rights Reserved.
 
 #include "ModelUnitTester.h"
-#include "NeuralNetworkInferenceTimer.h"
+#include "NeuralTimer.h"
 #include "NeuralNetworkInferenceQAUtils.h"
 #include "Misc/Paths.h"
 
@@ -27,7 +27,7 @@ FNNIUnitTesterTimeData FModelUnitTester_GetTimeInformation(UNeuralNetwork* InOut
 {
 	/* Input/output copy speed */ 
 	FNeuralStats OutCopyingStats;
-	FNeuralNetworkInferenceTimer Timer;
+	FNeuralTimer Timer;
 	for (int32 TimerIndex = 0; TimerIndex < InRepetitions; ++TimerIndex) 
 	{ 
 		InOutNetwork->SetInputFromArrayCopy(InInputArray); 
