*Module.h/cpp --> Module.h/cpp

#rnx

diff --git a/Engine/Plugins/Experimental/NeuralNetworkInference/Source/NeuralNetworkInference/Private/Module.cpp b/Engine/Plugins/Experimental/NeuralNetworkInference/Source/NeuralNetworkInference/Private/Module.cpp
new file mode 100644
--- /dev/null
+++ b/Engine/Plugins/Experimental/NeuralNetworkInference/Source/NeuralNetworkInference/Private/Module.cpp
@@ -0,0 +1,16 @@
+// Copyright Epic Games, Inc. All Rights Reserved.
+
+#include "Module.h"
+
+// This code will execute after your module is loaded into memory; the exact timing is specified in the .uplugin file per-module
+void FNeuralNetworkInferenceModule::StartupModule()
+{
+}
+
+// This function may be called during shutdown to clean up your module. For modules that support dynamic reloading,
+// we call this function before unloading the module.
+void FNeuralNetworkInferenceModule::ShutdownModule()
+{
+}
+
+IMPLEMENT_MODULE(FNeuralNetworkInferenceModule, NeuralNetworkInference);
diff --git a/Engine/Plugins/Experimental/NeuralNetworkInference/Source/NeuralNetworkInference/Private/Module.h b/Engine/Plugins/Experimental/NeuralNetworkInference/Source/NeuralNetworkInference/Private/Module.h
new file mode 100644
--- /dev/null
+++ b/Engine/Plugins/Experimental/NeuralNetworkInference/Source/NeuralNetworkInference/Private/Module.h
@@ -0,0 +1,14 @@
+// Copyright Epic Games, Inc. All Rights Reserved.
+
+#pragma once
+
+#include "CoreMinimal.h"
+#include "Modules/ModuleManager.h"
+
+class FNeuralNetworkInferenceModule : public IModuleInterface
+{
+public:
+	// Begin IModuleInterface
+	virtual void StartupModule() override;
+	virtual void ShutdownModule() override;
+};
diff --git a/Engine/Plugins/Experimental/NeuralNetworkInference/Source/NeuralNetworkInference/Private/NeuralNetworkInferenceModule.cpp b/Engine/Plugins/Experimental/NeuralNetworkInference/Source/NeuralNetworkInference/Private/NeuralNetworkInferenceModule.cpp
deleted file mode 100644
diff --git a/Engine/Plugins/Experimental/NeuralNetworkInference/Source/NeuralNetworkInference/Private/NeuralNetworkInferenceModule.h b/Engine/Plugins/Experimental/NeuralNetworkInference/Source/NeuralNetworkInference/Private/NeuralNetworkInferenceModule.h
deleted file mode 100644
diff --git a/Engine/Plugins/Experimental/NeuralNetworkInference/Source/NeuralNetworkInferenceCore/Private/Module.cpp b/Engine/Plugins/Experimental/NeuralNetworkInference/Source/NeuralNetworkInferenceCore/Private/Module.cpp
new file mode 100644
--- /dev/null
+++ b/Engine/Plugins/Experimental/NeuralNetworkInference/Source/NeuralNetworkInferenceCore/Private/Module.cpp
@@ -0,0 +1,16 @@
+// Copyright Epic Games, Inc. All Rights Reserved.
+
+#include "Module.h"
+
+// This code will execute after your module is loaded into memory; the exact timing is specified in the .uplugin file per-module
+void FNeuralNetworkInferenceCoreModule::StartupModule()
+{
+}
+
+// This function may be called during shutdown to clean up your module. For modules that support dynamic reloading,
+// we call this function before unloading the module.
+void FNeuralNetworkInferenceCoreModule::ShutdownModule()
+{
+}
+
+IMPLEMENT_MODULE(FNeuralNetworkInferenceCoreModule, NeuralNetworkInferenceCore);
diff --git a/Engine/Plugins/Experimental/NeuralNetworkInference/Source/NeuralNetworkInferenceCore/Private/Module.h b/Engine/Plugins/Experimental/NeuralNetworkInference/Source/NeuralNetworkInferenceCore/Private/Module.h
new file mode 100644
--- /dev/null
+++ b/Engine/Plugins/Experimental/NeuralNetworkInference/Source/NeuralNetworkInferenceCore/Private/Module.h
@@ -0,0 +1,13 @@
+// Copyright Epic Games, Inc. All Rights Reserved.
+
+#pragma once
+
+#include "Modules/ModuleManager.h"
+
+class FNeuralNetworkInferenceCoreModule : public IModuleInterface
+{
+public:
+	// Begin IModuleInterface
+	virtual void StartupModule() override;
+	virtual void ShutdownModule() override;
+};
diff --git a/Engine/Plugins/Experimental/NeuralNetworkInference/Source/NeuralNetworkInferenceCore/Private/NeuralNetworkInferenceCoreModule.cpp b/Engine/Plugins/Experimental/NeuralNetworkInference/Source/NeuralNetworkInferenceCore/Private/NeuralNetworkInferenceCoreModule.cpp
deleted file mode 100644
diff --git a/Engine/Plugins/Experimental/NeuralNetworkInference/Source/NeuralNetworkInferenceCore/Private/NeuralNetworkInferenceCoreModule.h b/Engine/Plugins/Experimental/NeuralNetworkInference/Source/NeuralNetworkInferenceCore/Private/NeuralNetworkInferenceCoreModule.h
deleted file mode 100644
diff --git a/Engine/Plugins/Experimental/NeuralNetworkInference/Source/NeuralNetworkInferenceQA/Private/Module.cpp b/Engine/Plugins/Experimental/NeuralNetworkInference/Source/NeuralNetworkInferenceQA/Private/Module.cpp
new file mode 100644
--- /dev/null
+++ b/Engine/Plugins/Experimental/NeuralNetworkInference/Source/NeuralNetworkInferenceQA/Private/Module.cpp
@@ -0,0 +1,16 @@
+// Copyright Epic Games, Inc. All Rights Reserved.
+
+#include "Module.h"
+
+// This code will execute after your module is loaded into memory; the exact timing is specified in the .uplugin file per-module
+void FNeuralNetworkInferenceQAModule::StartupModule()
+{
+}
+
+// This function may be called during shutdown to clean up your module. For modules that support dynamic reloading,
+// we call this function before unloading the module.
+void FNeuralNetworkInferenceQAModule::ShutdownModule()
+{
+}
+
+IMPLEMENT_MODULE(FNeuralNetworkInferenceQAModule, NeuralNetworkInferenceQA);
diff --git a/Engine/Plugins/Experimental/NeuralNetworkInference/Source/NeuralNetworkInferenceQA/Private/Module.h b/Engine/Plugins/Experimental/NeuralNetworkInference/Source/NeuralNetworkInferenceQA/Private/Module.h
new file mode 100644
--- /dev/null
+++ b/Engine/Plugins/Experimental/NeuralNetworkInference/Source/NeuralNetworkInferenceQA/Private/Module.h
@@ -0,0 +1,14 @@
+// Copyright Epic Games, Inc. All Rights Reserved.
+
+#pragma once
+
+#include "CoreMinimal.h"
+#include "Modules/ModuleManager.h"
+
+class FNeuralNetworkInferenceQAModule : public IModuleInterface
+{
+public:
+	// Begin IModuleInterface
+	virtual void StartupModule() override;
+	virtual void ShutdownModule() override;
+};
diff --git a/Engine/Plugins/Experimental/NeuralNetworkInference/Source/NeuralNetworkInferenceQA/Private/NeuralNetworkInferenceQAModule.cpp b/Engine/Plugins/Experimental/NeuralNetworkInference/Source/NeuralNetworkInferenceQA/Private/NeuralNetworkInferenceQAModule.cpp
deleted file mode 100644
diff --git a/Engine/Plugins/Experimental/NeuralNetworkInference/Source/NeuralNetworkInferenceQA/Private/NeuralNetworkInferenceQAModule.h b/Engine/Plugins/Experimental/NeuralNetworkInference/Source/NeuralNetworkInferenceQA/Private/NeuralNetworkInferenceQAModule.h
deleted file mode 100644
diff --git a/Engine/Plugins/Experimental/NeuralNetworkInference/Source/Private/ModelProto/Private/ModelProtoModule.cpp b/Engine/Plugins/Experimental/NeuralNetworkInference/Source/Private/ModelProto/Private/ModelProtoModule.cpp
deleted file mode 100644
diff --git a/Engine/Plugins/Experimental/NeuralNetworkInference/Source/Private/ModelProto/Private/ModelProtoModule.h b/Engine/Plugins/Experimental/NeuralNetworkInference/Source/Private/ModelProto/Private/ModelProtoModule.h
deleted file mode 100644
diff --git a/Engine/Plugins/Experimental/NeuralNetworkInference/Source/Private/ModelProto/Private/Module.cpp b/Engine/Plugins/Experimental/NeuralNetworkInference/Source/Private/ModelProto/Private/Module.cpp
new file mode 100644
--- /dev/null
+++ b/Engine/Plugins/Experimental/NeuralNetworkInference/Source/Private/ModelProto/Private/Module.cpp
@@ -0,0 +1,13 @@
+// Copyright Epic Games, Inc. All Rights Reserved.
+
+#include "Module.h"
+
+void FModelProtoModule::StartupModule()
+{
+}
+
+void FModelProtoModule::ShutdownModule()
+{
+}
+
+IMPLEMENT_MODULE(FModelProtoModule, ModelProto);
diff --git a/Engine/Plugins/Experimental/NeuralNetworkInference/Source/Private/ModelProto/Private/Module.h b/Engine/Plugins/Experimental/NeuralNetworkInference/Source/Private/ModelProto/Private/Module.h
new file mode 100644
--- /dev/null
+++ b/Engine/Plugins/Experimental/NeuralNetworkInference/Source/Private/ModelProto/Private/Module.h
@@ -0,0 +1,17 @@
+// Copyright Epic Games, Inc. All Rights Reserved.
+
+#pragma once
+
+#include "CoreMinimal.h"
+#include "Modules/ModuleManager.h"
+
+/**
+* The public interface to this module
+*/
+class FModelProtoModule : public IModuleInterface
+{
+public:
+	// Begin IModuleInterface
+	virtual void StartupModule() override;
+	virtual void ShutdownModule() override;
+};
diff --git a/Engine/Plugins/Experimental/NeuralNetworkInference/Source/Private/ModelProtoFileReader/Private/ModelProtoFileReaderModule.cpp b/Engine/Plugins/Experimental/NeuralNetworkInference/Source/Private/ModelProtoFileReader/Private/ModelProtoFileReaderModule.cpp
deleted file mode 100644
diff --git a/Engine/Plugins/Experimental/NeuralNetworkInference/Source/Private/ModelProtoFileReader/Private/ModelProtoFileReaderModule.h b/Engine/Plugins/Experimental/NeuralNetworkInference/Source/Private/ModelProtoFileReader/Private/ModelProtoFileReaderModule.h
deleted file mode 100644
diff --git a/Engine/Plugins/Experimental/NeuralNetworkInference/Source/Private/ModelProtoFileReader/Private/Module.cpp b/Engine/Plugins/Experimental/NeuralNetworkInference/Source/Private/ModelProtoFileReader/Private/Module.cpp
new file mode 100644
--- /dev/null
+++ b/Engine/Plugins/Experimental/NeuralNetworkInference/Source/Private/ModelProtoFileReader/Private/Module.cpp
@@ -0,0 +1,13 @@
+// Copyright Epic Games, Inc. All Rights Reserved.
+
+#include "Module.h"
+
+void FModelProtoFileReaderModule::StartupModule()
+{
+}
+
+void FModelProtoFileReaderModule::ShutdownModule()
+{
+}
+
+IMPLEMENT_MODULE(FModelProtoFileReaderModule, ModelProtoFileReader);
diff --git a/Engine/Plugins/Experimental/NeuralNetworkInference/Source/Private/ModelProtoFileReader/Private/Module.h b/Engine/Plugins/Experimental/NeuralNetworkInference/Source/Private/ModelProtoFileReader/Private/Module.h
new file mode 100644
--- /dev/null
+++ b/Engine/Plugins/Experimental/NeuralNetworkInference/Source/Private/ModelProtoFileReader/Private/Module.h
@@ -0,0 +1,17 @@
+// Copyright Epic Games, Inc. All Rights Reserved.
+
+#pragma once
+
+#include "CoreMinimal.h"
+#include "Modules/ModuleManager.h"
+
+/**
+* The public interface to this module
+*/
+class FModelProtoFileReaderModule : public IModuleInterface
+{
+public:
+	// Begin IModuleInterface
+	virtual void StartupModule() override;
+	virtual void ShutdownModule() override;
+};
diff --git a/Engine/Plugins/Experimental/NeuralNetworkInference/Source/Private/ThirdPartyHelperAndDLLLoader/Private/Module.cpp b/Engine/Plugins/Experimental/NeuralNetworkInference/Source/Private/ThirdPartyHelperAndDLLLoader/Private/Module.cpp
new file mode 100644
--- /dev/null
+++ b/Engine/Plugins/Experimental/NeuralNetworkInference/Source/Private/ThirdPartyHelperAndDLLLoader/Private/Module.cpp
@@ -0,0 +1,34 @@
+// Copyright Epic Games, Inc. All Rights Reserved.
+
+#include "Module.h"
+#include "HAL/PlatformProcess.h"
+#include "Interfaces/IPluginManager.h"
+
+// This code will execute after your module is loaded into memory; the exact timing is specified in the .uplugin file per-module
+void FThirdPartyHelperAndDLLLoaderModule::StartupModule()
+{
+#ifdef PLATFORM_WIN64
+	const FString PluginDir = IPluginManager::Get().FindPlugin(TEXT("NeuralNetworkInference"))->GetBaseDir();
+	const FString DirectMLRuntimeBinPath = PluginDir / TEXT(PREPROCESSOR_TO_STRING(DIRECTML_PLATFORM_PATH));
+	const FString DirectMLDLLPath = DirectMLRuntimeBinPath / TEXT(PREPROCESSOR_TO_STRING(DIRECTML_DLL_NAME));
+
+	FPlatformProcess::PushDllDirectory(*DirectMLRuntimeBinPath);
+	DirectMLDLLHandle = FPlatformProcess::GetDllHandle(*DirectMLDLLPath);
+	FPlatformProcess::PopDllDirectory(*DirectMLRuntimeBinPath);
+#endif //PLATFORM_WIN64
+}
+
+// This function may be called during shutdown to clean up your module. For modules that support dynamic reloading,
+// we call this function before unloading the module.
+void FThirdPartyHelperAndDLLLoaderModule::ShutdownModule()
+{
+#ifdef PLATFORM_WIN64
+	if (DirectMLDLLHandle)
+	{
+		FPlatformProcess::FreeDllHandle(DirectMLDLLHandle);
+		DirectMLDLLHandle = nullptr;
+	}
+#endif //PLATFORM_WIN64
+}
+
+IMPLEMENT_MODULE(FThirdPartyHelperAndDLLLoaderModule, ThirdPartyHelperAndDLLLoader);
diff --git a/Engine/Plugins/Experimental/NeuralNetworkInference/Source/Private/ThirdPartyHelperAndDLLLoader/Private/Module.h b/Engine/Plugins/Experimental/NeuralNetworkInference/Source/Private/ThirdPartyHelperAndDLLLoader/Private/Module.h
new file mode 100644
--- /dev/null
+++ b/Engine/Plugins/Experimental/NeuralNetworkInference/Source/Private/ThirdPartyHelperAndDLLLoader/Private/Module.h
@@ -0,0 +1,20 @@
+// Copyright Epic Games, Inc. All Rights Reserved.
+
+#pragma once
+
+#include "CoreMinimal.h"
+#include "Modules/ModuleManager.h"
+
+/**
+* The public interface to this module
+*/
+class FThirdPartyHelperAndDLLLoaderModule : public IModuleInterface
+{
+public:
+	// Begin IModuleInterface
+	virtual void StartupModule() override;
+	virtual void ShutdownModule() override;
+
+private:
+	void* DirectMLDLLHandle;
+};
diff --git a/Engine/Plugins/Experimental/NeuralNetworkInference/Source/Private/ThirdPartyHelperAndDLLLoader/Private/ThirdPartyHelperAndDLLLoaderModule.cpp b/Engine/Plugins/Experimental/NeuralNetworkInference/Source/Private/ThirdPartyHelperAndDLLLoader/Private/ThirdPartyHelperAndDLLLoaderModule.cpp
deleted file mode 100644
diff --git a/Engine/Plugins/Experimental/NeuralNetworkInference/Source/Private/ThirdPartyHelperAndDLLLoader/Private/ThirdPartyHelperAndDLLLoaderModule.h b/Engine/Plugins/Experimental/NeuralNetworkInference/Source/Private/ThirdPartyHelperAndDLLLoader/Private/ThirdPartyHelperAndDLLLoaderModule.h
deleted file mode 100644
