Towards removing NNICore - Part 1 - NeuralEnumClasses.h/cpp:
- From NNICore to NNI
- FDataType ranamed as FNeuralDataTypeUtils
- Made moves to allow this.

_codereview Francisco.Vicente
#rnx

diff --git a/Engine/Plugins/Experimental/NeuralNetworkInference/Source/NeuralNetworkInference/Private/NeuralEnumClasses.cpp b/Engine/Plugins/Experimental/NeuralNetworkInference/Source/NeuralNetworkInference/Private/NeuralEnumClasses.cpp
new file mode 100644
--- /dev/null
+++ b/Engine/Plugins/Experimental/NeuralNetworkInference/Source/NeuralNetworkInference/Private/NeuralEnumClasses.cpp
@@ -0,0 +1,83 @@
+// Copyright Epic Games, Inc. All Rights Reserved.
+
+#include "NeuralEnumClasses.h"
+#include "NeuralNetworkInferenceUtils.h"
+
+
+
+/* FNeuralNetworkInferenceUtils public functions
+ *****************************************************************************/
+
+FString FNeuralDataTypeUtils::ToString(const ENeuralDataType InDataType)
+{
+	if (InDataType == ENeuralDataType::Float)
+	{
+		return TEXT("Float");
+	}
+	else if (InDataType == ENeuralDataType::Int32)
+	{
+		return TEXT("Int32");
+	}
+	else if (InDataType == ENeuralDataType::Int64)
+	{
+		return TEXT("Int64");
+	}
+	else if (InDataType == ENeuralDataType::UInt32)
+	{
+		return TEXT("UInt32");
+	}
+	else if (InDataType == ENeuralDataType::UInt64)
+	{
+		return TEXT("UInt64");
+	}
+	else if (InDataType == ENeuralDataType::None)
+	{
+		return TEXT("None");
+	}
+	UE_LOG(LogNeuralNetworkInference, Warning, TEXT("FNeuralDataTypeUtils::ToString(): Unknown InDataType = %d used."), (int32)InDataType);
+	return TEXT("");
+}
+
+int64 FNeuralDataTypeUtils::GetSize(const ENeuralDataType InDataType)
+{
+	if (InDataType == ENeuralDataType::Float)
+	{
+		return sizeof(float);
+	}
+	else if (InDataType == ENeuralDataType::Int32)
+	{
+		return sizeof(int32);
+	}
+	else if (InDataType == ENeuralDataType::Int64)
+	{
+		return sizeof(int64);
+	}
+	else if (InDataType == ENeuralDataType::UInt32)
+	{
+		return sizeof(uint32);
+	}
+	else if (InDataType == ENeuralDataType::UInt64)
+	{
+		return sizeof(uint64);
+	}
+	UE_LOG(LogNeuralNetworkInference, Warning, TEXT("FNeuralDataTypeUtils::GetSize(): Unknown InDataType = %d used."), (int32)InDataType);
+	return 1;
+}
+
+EPixelFormat FNeuralDataTypeUtils::GetPixelFormat(const ENeuralDataType InDataType)
+{
+	if (InDataType == ENeuralDataType::Float)
+	{
+		return EPixelFormat::PF_R32_FLOAT;
+	}
+	else if (InDataType == ENeuralDataType::Int32)
+	{
+		return EPixelFormat::PF_R32_SINT;
+	}
+	else if (InDataType == ENeuralDataType::UInt32)
+	{
+		return EPixelFormat::PF_R32_UINT;
+	}
+	UE_LOG(LogNeuralNetworkInference, Warning, TEXT("FNeuralDataTypeUtils::GetPixelFormat(): Unknown InDataType = %d used."), (int32)InDataType);
+	return EPixelFormat::PF_Unknown;
+}
diff --git a/Engine/Plugins/Experimental/NeuralNetworkInference/Source/NeuralNetworkInference/Private/NeuralTensor.cpp b/Engine/Plugins/Experimental/NeuralNetworkInference/Source/NeuralNetworkInference/Private/NeuralTensor.cpp
--- a/Engine/Plugins/Experimental/NeuralNetworkInference/Source/NeuralNetworkInference/Private/NeuralTensor.cpp
+++ b/Engine/Plugins/Experimental/NeuralNetworkInference/Source/NeuralNetworkInference/Private/NeuralTensor.cpp
@@ -32,6 +32,13 @@ public:
 	{
 		constexpr T operator()(const T& InLeft, const T& InRight) const { return InLeft * InRight; }
 	};
+
+
+	/**
+	 * It gets the data type from its ETensorProtoDataType analog. E.g.,
+	 * const ENeuralDataType NeuralDataType = FModelProto::GetDataTypeFromTensorProtoDataType(TensorProto.DataType);
+	 */
+	static ENeuralDataType GetDataTypeFromTensorProtoDataType(const ETensorProtoDataType InTensorProtoDataType);
 };
 
 template <typename T>
@@ -92,6 +99,36 @@ void FPrivateNeuralTensor::NDTensorIndexesPlus1(TArray<int32>& InOutImageAreaInd
 	}
 }
 
+ENeuralDataType FPrivateNeuralTensor::GetDataTypeFromTensorProtoDataType(const ETensorProtoDataType InTensorProtoDataType)
+{
+	if (InTensorProtoDataType == ETensorProtoDataType::FLOAT)
+	{
+		return ENeuralDataType::Float;
+	}
+	else if (InTensorProtoDataType == ETensorProtoDataType::INT32)
+	{
+		return ENeuralDataType::Int32;
+	}
+	else if (InTensorProtoDataType == ETensorProtoDataType::INT64)
+	{
+		return ENeuralDataType::Int64;
+	}
+	else if (InTensorProtoDataType == ETensorProtoDataType::UINT32)
+	{
+		return ENeuralDataType::UInt32;
+	}
+	else if (InTensorProtoDataType == ETensorProtoDataType::UINT64)
+	{
+		return ENeuralDataType::UInt64;
+	}
+	else if (InTensorProtoDataType == ETensorProtoDataType::UNDEFINED)
+	{
+		return ENeuralDataType::None;
+	}
+	UE_LOG(LogNeuralNetworkInference, Warning, TEXT("FNeuralTensor::GetDataTypeFromTensorProtoDataType(): Unknown InTensorProtoDataType = %d used."), (int32)InTensorProtoDataType);
+	return ENeuralDataType::None;
+}
+
 
 
 /* FNeuralTensor costructors
@@ -292,7 +329,7 @@ void FNeuralTensor::ToGPU_RenderThread(FRDGBuilder* InOutGraphBuilder, const EBu
 	}
 	// Create BufferRef
 	FRDGBufferDesc BufferDesc;
-	BufferDesc.BytesPerElement = FDataType::GetSize(DataType);
+	BufferDesc.BytesPerElement = FNeuralDataTypeUtils::GetSize(DataType);
 	BufferDesc.NumElements = Num();
 	BufferDesc.Usage = InEBufferUsageFlags;
 	BufferDesc.UnderlyingType = FRDGBufferDesc::EUnderlyingType::VertexBuffer;
@@ -301,7 +338,7 @@ void FNeuralTensor::ToGPU_RenderThread(FRDGBuilder* InOutGraphBuilder, const EBu
 	// Recreate BufferSRVRef
 	if (EnumHasAnyFlags(InEBufferUsageFlags, BUF_ShaderResource))
 	{
-		BufferSRVRef = MakeShared<FRDGBufferSRVRef>(InOutGraphBuilder->CreateSRV(BufferRef, FDataType::GetPixelFormat(DataType)));
+		BufferSRVRef = MakeShared<FRDGBufferSRVRef>(InOutGraphBuilder->CreateSRV(BufferRef, FNeuralDataTypeUtils::GetPixelFormat(DataType)));
 	}
 	else
 	{
@@ -310,7 +347,7 @@ void FNeuralTensor::ToGPU_RenderThread(FRDGBuilder* InOutGraphBuilder, const EBu
 	// Recreate BufferUAVRef
 	if (EnumHasAnyFlags(InEBufferUsageFlags, BUF_UnorderedAccess))
 	{
-		BufferUAVRef = MakeShared<FRDGBufferUAVRef>(InOutGraphBuilder->CreateUAV(BufferRef, FDataType::GetPixelFormat(DataType)));
+		BufferUAVRef = MakeShared<FRDGBufferUAVRef>(InOutGraphBuilder->CreateUAV(BufferRef, FNeuralDataTypeUtils::GetPixelFormat(DataType)));
 	}
 	else
 	{
@@ -335,12 +372,12 @@ void FNeuralTensor::UpdateSRVAndOrUAV_RenderThread(FRDGBuilder* InOutGraphBuilde
 	// Recreate BufferSRVRef
 	if (BufferSRVRef.IsValid())
 	{
-		BufferSRVRef = MakeShared<FRDGBufferSRVRef>(InOutGraphBuilder->CreateSRV(BufferRef, FDataType::GetPixelFormat(DataType)));
+		BufferSRVRef = MakeShared<FRDGBufferSRVRef>(InOutGraphBuilder->CreateSRV(BufferRef, FNeuralDataTypeUtils::GetPixelFormat(DataType)));
 	}
 	// Recreate BufferUAVRef
 	if (BufferUAVRef.IsValid())
 	{
-		BufferUAVRef = MakeShared<FRDGBufferUAVRef>(InOutGraphBuilder->CreateUAV(BufferRef, FDataType::GetPixelFormat(DataType)));
+		BufferUAVRef = MakeShared<FRDGBufferUAVRef>(InOutGraphBuilder->CreateUAV(BufferRef, FNeuralDataTypeUtils::GetPixelFormat(DataType)));
 	}
 }
 
@@ -387,7 +424,7 @@ bool FNeuralTensor::InitPooledBuffer(void** NativeResource)
 			FRDGBuilder Builder(RHICmdList);
 
 			FRDGBufferDesc BufferDesc;
-			BufferDesc.BytesPerElement = FDataType::GetSize(DataType);
+			BufferDesc.BytesPerElement = FNeuralDataTypeUtils::GetSize(DataType);
 			BufferDesc.NumElements = Num();
 			BufferDesc.Usage = BUF_UnorderedAccess | BUF_ShaderResource;
 			BufferDesc.UnderlyingType = FRDGBufferDesc::EUnderlyingType::VertexBuffer;
@@ -399,8 +436,8 @@ bool FNeuralTensor::InitPooledBuffer(void** NativeResource)
 			*PooledBuffer = Builder.ConvertToExternalBuffer(BufferRef);
 
 			// SRV and UAV can be used by callers (for example MLDeformer), look at the GetPooledBuffer()
-			BufferSRVRef = MakeShared<FRDGBufferSRVRef>(Builder.CreateSRV(BufferRef, FDataType::GetPixelFormat(DataType)));
-			BufferUAVRef = MakeShared<FRDGBufferUAVRef>(Builder.CreateUAV(BufferRef, FDataType::GetPixelFormat(DataType)));
+			BufferSRVRef = MakeShared<FRDGBufferSRVRef>(Builder.CreateSRV(BufferRef, FNeuralDataTypeUtils::GetPixelFormat(DataType)));
+			BufferUAVRef = MakeShared<FRDGBufferUAVRef>(Builder.CreateUAV(BufferRef, FNeuralDataTypeUtils::GetPixelFormat(DataType)));
 
 			Builder.Execute();
 
@@ -473,7 +510,7 @@ void FNeuralTensor::SetNumUninitialized(const TArray<int64>& InSizes, const ENeu
 	if (Sizes.Num() > 0)
 	{
 		Volume = FNeuralNetworkInferenceUtils::Product<int64>(Sizes);
-		const int64 VolumeInBytes = Volume * FDataType::GetSize(DataType);
+		const int64 VolumeInBytes = Volume * FNeuralDataTypeUtils::GetSize(DataType);
 		if (VolumeInBytes != ArrayCPU.Num())
 		{
 			ArrayCPU.SetNumUninitialized(VolumeInBytes, bInAllowShrinking); // Pre-allocate TArray
@@ -511,7 +548,7 @@ bool FNeuralTensor::SetFromTensorProto(const FTensorProto* const InTensorProto,
 	Name = InTensorName;
 	TensorTypeGPU = InTensorTypeGPU;
 	// Memory allocation
-	SetNumUninitialized(InTensorProto->Dimensions, InTensorProto->GetDataTypeFromTensorProtoDataType());
+	SetNumUninitialized(InTensorProto->Dimensions, FPrivateNeuralTensor::GetDataTypeFromTensorProtoDataType(InTensorProto->DataType));
 
 	// RawData
 	if (!InTensorProto->RawData.IsEmpty())
@@ -532,7 +569,7 @@ bool FNeuralTensor::SetFromTensorProto(const FTensorProto* const InTensorProto,
 		// }
 		// // Read neural tensor from binary data
 		// const FString BinaryWeightFilePath = InExternalDataDirectory / InTensorProto->ExternalData[0].Value;
-		// if (!FModelProtoFileReader::ReadWeightsFromOtxtBinaryFile((char*)GetData(), Num() * FDataType::GetSize(NeuralDataType), BinaryWeightFilePath))
+		// if (!FModelProtoFileReader::ReadWeightsFromOtxtBinaryFile((char*)GetData(), Num() * FNeuralDataTypeUtils::GetSize(NeuralDataType), BinaryWeightFilePath))
 		// {
 		// 	UE_LOG(LogNeuralNetworkInference, Warning, TEXT("FNeuralTensor::SetFromTensorProto(): Could not read binary file: %s."), *BinaryWeightFilePath);
 		// 	return false;
@@ -565,7 +602,7 @@ bool FNeuralTensor::Flip(const int32 InDimension)
 	{
 		DimensionOffset *= Sizes[DimensionIndex];
 	}
-	const int64 BytesPerIndex = FDataType::GetSize(DataType);
+	const int64 BytesPerIndex = FNeuralDataTypeUtils::GetSize(DataType);
 	const int64 DimensionOffsetInBytes = DimensionOffset * BytesPerIndex;
 
 	// Fill TensorNDIndexes and TensorNDSizes
@@ -649,7 +686,7 @@ bool FNeuralTensor::Transpose()
 			const uint8* const ArrayOnCPUPtr = ArrayCPU.GetData();
 			const int64 Height = Sizes[0];
 			const int64 Width = Sizes[1];
-			const int64 Bytes = FDataType::GetSize(DataType);
+			const int64 Bytes = FNeuralDataTypeUtils::GetSize(DataType);
 			for (int64 Y = 0; Y < Height; ++Y)
 			{
 				for (int64 X = 0; X < Width; ++X)
@@ -698,7 +735,7 @@ FString FNeuralTensor::ToString(const int64 InMaxNumberElementsToDisplay, const
 	{
 		TensorString += (Name.Len() > 0 ? Name : FString(TEXT("Unnamed FNeuralTensor"))) + TEXT(": ")
 			// DataType
-			+ FDataType::ToString(DataType) + TEXT(", ")
+			+ FNeuralDataTypeUtils::ToString(DataType) + TEXT(", ")
 			+ (TensorTypeGPU == ENeuralTensorTypeGPU::Generic ? TEXT("Generic")
 				: TensorTypeGPU == ENeuralTensorTypeGPU::Input ? TEXT("Input")
 				: TensorTypeGPU == ENeuralTensorTypeGPU::IntermediateNotInitialized ? TEXT("IntermediateNotInitialized")
@@ -826,10 +863,10 @@ void FNeuralTensor::SetFromPointer(const void* const InData, const int64 InSizeO
 
 bool FNeuralTensor::CheckTAndDataTypeResult(const bool bInCheckTAndDataTypeResult, const int64 InSizeOfT) const
 {
-	const int64 SizeOfDataType = FDataType::GetSize(DataType);
+	const int64 SizeOfDataType = FNeuralDataTypeUtils::GetSize(DataType);
 	if (!bInCheckTAndDataTypeResult)
 	{
-		const FString DataTypeString = FDataType::ToString(DataType);
+		const FString DataTypeString = FNeuralDataTypeUtils::ToString(DataType);
 		// sizeof(T) and DataType do not match
 		if (SizeOfDataType != InSizeOfT)
 		{
diff --git a/Engine/Plugins/Experimental/NeuralNetworkInference/Source/NeuralNetworkInference/Public/NeuralEnumClasses.h b/Engine/Plugins/Experimental/NeuralNetworkInference/Source/NeuralNetworkInference/Public/NeuralEnumClasses.h
new file mode 100644
--- /dev/null
+++ b/Engine/Plugins/Experimental/NeuralNetworkInference/Source/NeuralNetworkInference/Public/NeuralEnumClasses.h
@@ -0,0 +1,106 @@
+// Copyright Epic Games, Inc. All Rights Reserved.
+
+#pragma once
+
+#include "CoreMinimal.h"
+#include "PixelFormat.h"
+#include <type_traits>
+#include "NeuralEnumClasses.generated.h"
+
+/**
+ * Whether UNeuralNetwork will use the highly optimized UnrealEngine-and-ONNXRuntime-based back end (UEAndORT) or the less optimized but fully cross platform only-UE one (UEOnly).
+ * We recommend using Auto, which will find and use the optimal back end for each platform.
+ */
+UENUM()
+enum class ENeuralBackEnd : uint8
+{
+	UEAndORT, /* UnrealEngine-and-ONNXRuntime-accelerated back end, ideal for those platforms that support it. WITH_UE_AND_ORT_SUPPORT is the C++ define macro that checks whether UEAndORT support exists for the current platform.*/
+	UEOnly, /* It might be slower than the UEAndORT back end, but it will compile in all platforms and OSs compatible with Unreal Engine. */
+	Auto /* Recommended setting. It will use UEAndORT if possible, and fall back to UEOnly only for those platforms that do not support UEAndORT yet. */
+};
+
+UENUM()
+enum class ENeuralDataType : uint8
+{
+	Float,
+	//Double,
+	//Int8,
+	//Int16,
+	Int32,
+	Int64,
+	//UInt8,
+	UInt32,
+	UInt64,
+	None
+};
+
+UENUM()
+enum class ENeuralDeviceType : uint8
+{
+	CPU,
+	GPU,
+	None
+};
+
+
+
+/**
+ * Auxiliary utils class for ENeuralDataType
+ */
+
+class NEURALNETWORKINFERENCE_API FNeuralDataTypeUtils
+{
+public:
+	static FString ToString(const ENeuralDataType InDataType);
+	static int64 GetSize(const ENeuralDataType InDataType);
+	static EPixelFormat GetPixelFormat(const ENeuralDataType InDataType);
+
+	/**
+	 * It checks whether T and InDataType are the same type. E.g.,
+	 * checkf(CheckTAndDataType<float>(), TEXT("Expected a ENeuralDataType::Float type."));
+	 */
+	template <typename T>
+	static bool CheckTAndDataType(const ENeuralDataType InDataType)
+	{
+		return InDataType == GetDataType<T>();
+	}
+
+	/**
+	 * It gets the data type from the type T. E.g.,
+	 * checkf(InDataType == GetDataType<float>(), TEXT("InDataType == GetDataType<float>() failed!"))
+	 * FNeuralTensor(FNeuralDataTypeUtils::GetDataType<T>(), InArray.GetData(), ...)
+	 */
+	template <typename T>
+	static ENeuralDataType GetDataType();
+};
+
+
+
+/* FNeuralDataTypeUtils templated functions
+ *****************************************************************************/
+
+template <typename T>
+ENeuralDataType FNeuralDataTypeUtils::GetDataType()
+{
+	if (std::is_same<T, float>::value)
+	{
+		return ENeuralDataType::Float;
+	}
+	else if (std::is_same<T, int32>::value)
+	{
+		return ENeuralDataType::Int32;
+	}
+	else if (std::is_same<T, int64>::value)
+	{
+		return ENeuralDataType::Int64;
+	}
+	else if (std::is_same<T, uint32>::value)
+	{
+		return ENeuralDataType::UInt32;
+	}
+	else if (std::is_same<T, uint64>::value)
+	{
+		return ENeuralDataType::UInt64;
+	}
+	return ENeuralDataType::None;
+}
diff --git a/Engine/Plugins/Experimental/NeuralNetworkInference/Source/NeuralNetworkInference/Public/NeuralTensor.h b/Engine/Plugins/Experimental/NeuralNetworkInference/Source/NeuralNetworkInference/Public/NeuralTensor.h
--- a/Engine/Plugins/Experimental/NeuralNetworkInference/Source/NeuralNetworkInference/Public/NeuralTensor.h
+++ b/Engine/Plugins/Experimental/NeuralNetworkInference/Source/NeuralNetworkInference/Public/NeuralTensor.h
@@ -330,7 +330,7 @@ private:
 
 template <typename T>
 FNeuralTensor::FNeuralTensor(const TArray<T>& InArray, const TArray<int64>& InSizes, const FString& InName, const ENeuralTensorTypeGPU InTensorTypeGPU)
-	: FNeuralTensor(FDataType::GetDataType<T>(), InArray.GetData(), sizeof(T), InArray.Num(), (InSizes.Num() > 0 ? InSizes : TArray<int64>({ (int64)InArray.Num() })), InName, InTensorTypeGPU)
+	: FNeuralTensor(FNeuralDataTypeUtils::GetDataType<T>(), InArray.GetData(), sizeof(T), InArray.Num(), (InSizes.Num() > 0 ? InSizes : TArray<int64>({ (int64)InArray.Num() })), InName, InTensorTypeGPU)
 {}
 
 bool FNeuralTensor::operator!=(const FNeuralTensor& InTensorToCopy) const
@@ -464,5 +464,5 @@ void FNeuralTensor::SetTo(const TInput InValue)
 template<typename T>
 bool FNeuralTensor::CheckTAndDataType() const
 {
-	return CheckTAndDataTypeResult(FDataType::CheckTAndDataType<T>(DataType), sizeof(T));
+	return CheckTAndDataTypeResult(FNeuralDataTypeUtils::CheckTAndDataType<T>(DataType), sizeof(T));
 }
diff --git a/Engine/Plugins/Experimental/NeuralNetworkInference/Source/NeuralNetworkInferenceCore/Private/NeuralEnumClasses.cpp b/Engine/Plugins/Experimental/NeuralNetworkInference/Source/NeuralNetworkInferenceCore/Private/NeuralEnumClasses.cpp
deleted file mode 100644
diff --git a/Engine/Plugins/Experimental/NeuralNetworkInference/Source/NeuralNetworkInferenceCore/Public/NeuralEnumClasses.h b/Engine/Plugins/Experimental/NeuralNetworkInference/Source/NeuralNetworkInferenceCore/Public/NeuralEnumClasses.h
deleted file mode 100644
diff --git a/Engine/Plugins/Experimental/NeuralNetworkInference/Source/NeuralNetworkInferenceQA/NeuralNetworkInferenceQA.Build.cs b/Engine/Plugins/Experimental/NeuralNetworkInference/Source/NeuralNetworkInferenceQA/NeuralNetworkInferenceQA.Build.cs
--- a/Engine/Plugins/Experimental/NeuralNetworkInference/Source/NeuralNetworkInferenceQA/NeuralNetworkInferenceQA.Build.cs
+++ b/Engine/Plugins/Experimental/NeuralNetworkInference/Source/NeuralNetworkInferenceQA/NeuralNetworkInferenceQA.Build.cs
@@ -30,7 +30,6 @@ public class NeuralNetworkInferenceQA : ModuleRules
 			new string[] {
 				"ModelProto",
 				"ModelProtoFileReader",
-				"NeuralNetworkInferenceCore",
 				"NeuralNetworkInferenceProfiling"
 			}
 		);
diff --git a/Engine/Plugins/Experimental/NeuralNetworkInference/Source/Private/ModelProto/Internal/ModelProto.h b/Engine/Plugins/Experimental/NeuralNetworkInference/Source/Private/ModelProto/Internal/ModelProto.h
--- a/Engine/Plugins/Experimental/NeuralNetworkInference/Source/Private/ModelProto/Internal/ModelProto.h
+++ b/Engine/Plugins/Experimental/NeuralNetworkInference/Source/Private/ModelProto/Internal/ModelProto.h
@@ -3,7 +3,6 @@
 #pragma once
 
 #include "CoreMinimal.h"
-#include "NeuralEnumClasses.h"
 #include "ModelProto.generated.h"
 
 
@@ -250,12 +249,6 @@ public:
 
 	FString ToString(const FString& InLineStarted = TEXT("")) const;
 
-	/**
-	 * It gets the data type from its ETensorProtoDataType analog. E.g.,
-	 * const ENeuralDataType NeuralDataType = FModelProto::GetDataTypeFromTensorProtoDataType(TensorProto.DataType);
-	 */
-	ENeuralDataType GetDataTypeFromTensorProtoDataType() const;
-
 private:
 	static ETensorProtoDataType StringToDataType(const FString& InString);
 
diff --git a/Engine/Plugins/Experimental/NeuralNetworkInference/Source/Private/ModelProto/Private/ModelProto.cpp b/Engine/Plugins/Experimental/NeuralNetworkInference/Source/Private/ModelProto/Private/ModelProto.cpp
--- a/Engine/Plugins/Experimental/NeuralNetworkInference/Source/Private/ModelProto/Private/ModelProto.cpp
+++ b/Engine/Plugins/Experimental/NeuralNetworkInference/Source/Private/ModelProto/Private/ModelProto.cpp
@@ -230,36 +230,6 @@ FString FTensorProto::ToString(const FString& InLineStarted) const
 	return String;
 }
 
-ENeuralDataType FTensorProto::GetDataTypeFromTensorProtoDataType() const
-{
-	if (DataType == ETensorProtoDataType::FLOAT)
-	{
-		return ENeuralDataType::Float;
-	}
-	else if (DataType == ETensorProtoDataType::INT32)
-	{
-		return ENeuralDataType::Int32;
-	}
-	else if (DataType == ETensorProtoDataType::INT64)
-	{
-		return ENeuralDataType::Int64;
-	}
-	else if (DataType == ETensorProtoDataType::UINT32)
-	{
-		return ENeuralDataType::UInt32;
-	}
-	else if (DataType == ETensorProtoDataType::UINT64)
-	{
-		return ENeuralDataType::UInt64;
-	}
-	else if (DataType == ETensorProtoDataType::UNDEFINED)
-	{
-		return ENeuralDataType::None;
-	}
-	UE_LOG(LogModelProto, Warning, TEXT("FModelProto::GetDataTypeFromTensorProtoDataType(): Unknown DataType = %d used."), (int32)DataType);
-	return ENeuralDataType::None;
-}
-
 /* FTensorProto private functions
  *****************************************************************************/
 
