Towards removing NNICore - Part 2 - NeuralOperatorEnumClasses.h/cpp:
- NeuralNetworkInferenceCore module removed.
- NeuralOperatorEnumClasses.h moved to NNIShaders.
- NNIShaders from Public to Private in NNI.Build.cs.

_codereview Francisco.Vicente
#rnx

diff --git a/Engine/Plugins/Experimental/NeuralNetworkInference/NeuralNetworkInference.uplugin b/Engine/Plugins/Experimental/NeuralNetworkInference/NeuralNetworkInference.uplugin
--- a/Engine/Plugins/Experimental/NeuralNetworkInference/NeuralNetworkInference.uplugin
+++ b/Engine/Plugins/Experimental/NeuralNetworkInference/NeuralNetworkInference.uplugin
@@ -21,11 +21,6 @@
 			"Type": "RuntimeAndProgram",
 			"LoadingPhase": "Default"
 		},
-		{
-			"Name": "NeuralNetworkInferenceCore",
-			"Type": "RuntimeAndProgram",
-			"LoadingPhase": "Default"
-		},
 		{
 			"Name": "NeuralNetworkInferenceShaders",
 			"Type": "RuntimeAndProgram",
diff --git a/Engine/Plugins/Experimental/NeuralNetworkInference/Source/NeuralNetworkInference/Internal/NeuralOperators/ElementWiseOperator.h b/Engine/Plugins/Experimental/NeuralNetworkInference/Source/NeuralNetworkInference/Internal/NeuralOperators/ElementWiseOperator.h
--- a/Engine/Plugins/Experimental/NeuralNetworkInference/Source/NeuralNetworkInference/Internal/NeuralOperators/ElementWiseOperator.h
+++ b/Engine/Plugins/Experimental/NeuralNetworkInference/Source/NeuralNetworkInference/Internal/NeuralOperators/ElementWiseOperator.h
@@ -4,13 +4,14 @@
 
 #include "CoreMinimal.h"
 #include "NeuralOperator.h"
-#include "NeuralOperatorEnumClasses.h"
 #include "NeuralTensor.h"
 
+enum class EElementWiseOperator : uint8;
+
 class NEURALNETWORKINFERENCE_API IElementWiseOperator : public FNeuralOperator
 {
 public:
-	IElementWiseOperator(const FString& InName, const int32 InVersion, const EElementWiseOperator InElementWiseOperator, const bool bIsInlinedTensor,
+	IElementWiseOperator(const FString& InName, const int32 InVersion, const TSharedPtr<EElementWiseOperator>& InElementWiseOperator, const bool bIsInlinedTensor,
 		const TArray<float>& InAttributes = TArray<float>());
 
 	virtual ~IElementWiseOperator();
@@ -33,5 +34,5 @@ protected:
 	virtual void ForwardCPUWithFunction(float InOperatorFunction(const float, const float)) final;
 
 private:
-	const EElementWiseOperator ElementWiseOperator;
+	const TSharedPtr<EElementWiseOperator> ElementWiseOperator;
 };
diff --git a/Engine/Plugins/Experimental/NeuralNetworkInference/Source/NeuralNetworkInference/Internal/NeuralOperators/MultidirectionalBroadcastOperator.h b/Engine/Plugins/Experimental/NeuralNetworkInference/Source/NeuralNetworkInference/Internal/NeuralOperators/MultidirectionalBroadcastOperator.h
--- a/Engine/Plugins/Experimental/NeuralNetworkInference/Source/NeuralNetworkInference/Internal/NeuralOperators/MultidirectionalBroadcastOperator.h
+++ b/Engine/Plugins/Experimental/NeuralNetworkInference/Source/NeuralNetworkInference/Internal/NeuralOperators/MultidirectionalBroadcastOperator.h
@@ -5,8 +5,9 @@
 #include "CoreMinimal.h"
 #include "NeuralInt64ArrayUInt32Buffer.h"
 #include "NeuralOperator.h"
-#include "NeuralOperatorEnumClasses.h"
 
+enum class EMultidirectionalBroadcastOperator : uint8;
+enum class EMultidirectionalBroadcastShapeMode : uint8;
 struct FReadBuffer;
 
 class NEURALNETWORKINFERENCE_API IMultidirectionalBroadcastOperator : public FNeuralOperator
@@ -16,7 +17,7 @@ public:
 	 * Not inlined layer if an empty InPotentialInlinedTensors is given.
 	 * It will try to inline the input tensors otherwise, where only the input tensors with the given indexes (InPotentialInlinedTensors) can be potentially inlined layers.
 	 */
-	IMultidirectionalBroadcastOperator(const FString& InName, const int32 InVersion, const EMultidirectionalBroadcastOperator InMultidirectionalBroadcastOperator,
+	IMultidirectionalBroadcastOperator(const FString& InName, const int32 InVersion, const TSharedPtr<EMultidirectionalBroadcastOperator>& InMultidirectionalBroadcastOperator,
 		const TSet<uint32>& InPotentialInlinedTensors);
 
 	virtual ~IMultidirectionalBroadcastOperator();
@@ -38,8 +39,8 @@ protected:
 	virtual void ForwardCPUWithFunction(float InOperatorFunction(const float, const float)) final;
 
 private:
-	const EMultidirectionalBroadcastOperator MultidirectionalBroadcastOperator;
-	EMultidirectionalBroadcastShapeMode MultidirectionalBroadcastShapeMode;
+	const TSharedPtr<EMultidirectionalBroadcastOperator> MultidirectionalBroadcastOperator;
+	const TSharedPtr<EMultidirectionalBroadcastShapeMode> MultidirectionalBroadcastShapeMode;
 	FNeuralInt64ArrayUInt32Buffer ShapesX;
 	FNeuralInt64ArrayUInt32Buffer ShapesY;
 	FNeuralInt64ArrayUInt32Buffer ShapesOutput;
diff --git a/Engine/Plugins/Experimental/NeuralNetworkInference/Source/NeuralNetworkInference/NeuralNetworkInference.Build.cs b/Engine/Plugins/Experimental/NeuralNetworkInference/Source/NeuralNetworkInference/NeuralNetworkInference.Build.cs
--- a/Engine/Plugins/Experimental/NeuralNetworkInference/Source/NeuralNetworkInference/NeuralNetworkInference.Build.cs
+++ b/Engine/Plugins/Experimental/NeuralNetworkInference/Source/NeuralNetworkInference/NeuralNetworkInference.Build.cs
@@ -29,7 +29,6 @@ public class NeuralNetworkInference : ModuleRules
 				"Core",
 				"CoreUObject",
 				"Engine",
-				"NeuralNetworkInferenceCore",
 				"NeuralNetworkInferenceProfiling",
 				// Internal (Backend-related)
 				"ModelProto"
diff --git a/Engine/Plugins/Experimental/NeuralNetworkInference/Source/NeuralNetworkInference/Private/NeuralOperators/AbsOperator.cpp b/Engine/Plugins/Experimental/NeuralNetworkInference/Source/NeuralNetworkInference/Private/NeuralOperators/AbsOperator.cpp
--- a/Engine/Plugins/Experimental/NeuralNetworkInference/Source/NeuralNetworkInference/Private/NeuralOperators/AbsOperator.cpp
+++ b/Engine/Plugins/Experimental/NeuralNetworkInference/Source/NeuralNetworkInference/Private/NeuralOperators/AbsOperator.cpp
@@ -1,6 +1,7 @@
 // Copyright Epic Games, Inc. All Rights Reserved.
 
 #include "NeuralOperators/AbsOperator.h"
+#include "NeuralOperatorEnumClasses.h"
 
 
 
@@ -8,7 +9,7 @@
  *****************************************************************************/
 
 FAbsOperator::FAbsOperator(const bool bIsInlinedTensor)
-	: IElementWiseOperator(TEXT("Abs"), 13, EElementWiseOperator::Abs, bIsInlinedTensor)
+	: IElementWiseOperator(TEXT("Abs"), 13, MakeShared<EElementWiseOperator>(EElementWiseOperator::Abs), bIsInlinedTensor)
 {
 }
 
diff --git a/Engine/Plugins/Experimental/NeuralNetworkInference/Source/NeuralNetworkInference/Private/NeuralOperators/AcosOperator.cpp b/Engine/Plugins/Experimental/NeuralNetworkInference/Source/NeuralNetworkInference/Private/NeuralOperators/AcosOperator.cpp
--- a/Engine/Plugins/Experimental/NeuralNetworkInference/Source/NeuralNetworkInference/Private/NeuralOperators/AcosOperator.cpp
+++ b/Engine/Plugins/Experimental/NeuralNetworkInference/Source/NeuralNetworkInference/Private/NeuralOperators/AcosOperator.cpp
@@ -1,6 +1,7 @@
 // Copyright Epic Games, Inc. All Rights Reserved.
 
 #include "NeuralOperators/AcosOperator.h"
+#include "NeuralOperatorEnumClasses.h"
 
 
 
@@ -8,7 +9,7 @@
  *****************************************************************************/
 
 FAcosOperator::FAcosOperator(const bool bIsInlinedTensor)
-	: IElementWiseOperator(TEXT("Acos"), 13, EElementWiseOperator::Acos, bIsInlinedTensor)
+	: IElementWiseOperator(TEXT("Acos"), 13, MakeShared<EElementWiseOperator>(EElementWiseOperator::Acos), bIsInlinedTensor)
 {
 }
 
diff --git a/Engine/Plugins/Experimental/NeuralNetworkInference/Source/NeuralNetworkInference/Private/NeuralOperators/AddOperator.cpp b/Engine/Plugins/Experimental/NeuralNetworkInference/Source/NeuralNetworkInference/Private/NeuralOperators/AddOperator.cpp
--- a/Engine/Plugins/Experimental/NeuralNetworkInference/Source/NeuralNetworkInference/Private/NeuralOperators/AddOperator.cpp
+++ b/Engine/Plugins/Experimental/NeuralNetworkInference/Source/NeuralNetworkInference/Private/NeuralOperators/AddOperator.cpp
@@ -1,6 +1,7 @@
 // Copyright Epic Games, Inc. All Rights Reserved.
 
 #include "NeuralOperators/AddOperator.h"
+#include "NeuralOperatorEnumClasses.h"
 
 
 
@@ -8,7 +9,7 @@
  *****************************************************************************/
 
 FAddOperator::FAddOperator(const TSet<uint32>& InPotentialInlinedTensors)
-	: IMultidirectionalBroadcastOperator(TEXT("Add"), 13, EMultidirectionalBroadcastOperator::Add, InPotentialInlinedTensors)
+	: IMultidirectionalBroadcastOperator(TEXT("Add"), 13, MakeShared<EMultidirectionalBroadcastOperator>(EMultidirectionalBroadcastOperator::Add), InPotentialInlinedTensors)
 {
 }
 
diff --git a/Engine/Plugins/Experimental/NeuralNetworkInference/Source/NeuralNetworkInference/Private/NeuralOperators/AsinOperator.cpp b/Engine/Plugins/Experimental/NeuralNetworkInference/Source/NeuralNetworkInference/Private/NeuralOperators/AsinOperator.cpp
--- a/Engine/Plugins/Experimental/NeuralNetworkInference/Source/NeuralNetworkInference/Private/NeuralOperators/AsinOperator.cpp
+++ b/Engine/Plugins/Experimental/NeuralNetworkInference/Source/NeuralNetworkInference/Private/NeuralOperators/AsinOperator.cpp
@@ -1,6 +1,7 @@
 // Copyright Epic Games, Inc. All Rights Reserved.
 
 #include "NeuralOperators/AsinOperator.h"
+#include "NeuralOperatorEnumClasses.h"
 
 
 
@@ -8,7 +9,7 @@
  *****************************************************************************/
 
 FAsinOperator::FAsinOperator(const bool bIsInlinedTensor)
-	: IElementWiseOperator(TEXT("Asin"), 7, EElementWiseOperator::Asin, bIsInlinedTensor)
+	: IElementWiseOperator(TEXT("Asin"), 7, MakeShared<EElementWiseOperator>(EElementWiseOperator::Asin), bIsInlinedTensor)
 {
 }
 
diff --git a/Engine/Plugins/Experimental/NeuralNetworkInference/Source/NeuralNetworkInference/Private/NeuralOperators/AtanOperator.cpp b/Engine/Plugins/Experimental/NeuralNetworkInference/Source/NeuralNetworkInference/Private/NeuralOperators/AtanOperator.cpp
--- a/Engine/Plugins/Experimental/NeuralNetworkInference/Source/NeuralNetworkInference/Private/NeuralOperators/AtanOperator.cpp
+++ b/Engine/Plugins/Experimental/NeuralNetworkInference/Source/NeuralNetworkInference/Private/NeuralOperators/AtanOperator.cpp
@@ -1,6 +1,7 @@
 // Copyright Epic Games, Inc. All Rights Reserved.
 
 #include "NeuralOperators/AtanOperator.h"
+#include "NeuralOperatorEnumClasses.h"
 
 
 
@@ -8,7 +9,7 @@
  *****************************************************************************/
 
 FAtanOperator::FAtanOperator(const bool bIsInlinedTensor)
-	: IElementWiseOperator(TEXT("Atan"), 7, EElementWiseOperator::Atan, bIsInlinedTensor)
+	: IElementWiseOperator(TEXT("Atan"), 7, MakeShared<EElementWiseOperator>(EElementWiseOperator::Atan), bIsInlinedTensor)
 {
 }
 
diff --git a/Engine/Plugins/Experimental/NeuralNetworkInference/Source/NeuralNetworkInference/Private/NeuralOperators/CeilOperator.cpp b/Engine/Plugins/Experimental/NeuralNetworkInference/Source/NeuralNetworkInference/Private/NeuralOperators/CeilOperator.cpp
--- a/Engine/Plugins/Experimental/NeuralNetworkInference/Source/NeuralNetworkInference/Private/NeuralOperators/CeilOperator.cpp
+++ b/Engine/Plugins/Experimental/NeuralNetworkInference/Source/NeuralNetworkInference/Private/NeuralOperators/CeilOperator.cpp
@@ -1,6 +1,7 @@
 // Copyright Epic Games, Inc. All Rights Reserved.
 
 #include "NeuralOperators/CeilOperator.h"
+#include "NeuralOperatorEnumClasses.h"
 
 
 
@@ -8,7 +9,7 @@
  *****************************************************************************/
 
 FCeilOperator::FCeilOperator(const bool bIsInlinedTensor)
-	: IElementWiseOperator(TEXT("Ceil"), 13, EElementWiseOperator::Ceil, bIsInlinedTensor)
+	: IElementWiseOperator(TEXT("Ceil"), 13, MakeShared<EElementWiseOperator>(EElementWiseOperator::Ceil), bIsInlinedTensor)
 {
 }
 
diff --git a/Engine/Plugins/Experimental/NeuralNetworkInference/Source/NeuralNetworkInference/Private/NeuralOperators/CosOperator.cpp b/Engine/Plugins/Experimental/NeuralNetworkInference/Source/NeuralNetworkInference/Private/NeuralOperators/CosOperator.cpp
--- a/Engine/Plugins/Experimental/NeuralNetworkInference/Source/NeuralNetworkInference/Private/NeuralOperators/CosOperator.cpp
+++ b/Engine/Plugins/Experimental/NeuralNetworkInference/Source/NeuralNetworkInference/Private/NeuralOperators/CosOperator.cpp
@@ -1,6 +1,7 @@
 // Copyright Epic Games, Inc. All Rights Reserved.
 
 #include "NeuralOperators/CosOperator.h"
+#include "NeuralOperatorEnumClasses.h"
 
 
 
@@ -8,7 +9,7 @@
  *****************************************************************************/
 
 FCosOperator::FCosOperator(const bool bIsInlinedTensor)
-	: IElementWiseOperator(TEXT("Cos"), 7, EElementWiseOperator::Cos, bIsInlinedTensor)
+	: IElementWiseOperator(TEXT("Cos"), 7, MakeShared<EElementWiseOperator>(EElementWiseOperator::Cos), bIsInlinedTensor)
 {
 }
 
diff --git a/Engine/Plugins/Experimental/NeuralNetworkInference/Source/NeuralNetworkInference/Private/NeuralOperators/CoshOperator.cpp b/Engine/Plugins/Experimental/NeuralNetworkInference/Source/NeuralNetworkInference/Private/NeuralOperators/CoshOperator.cpp
--- a/Engine/Plugins/Experimental/NeuralNetworkInference/Source/NeuralNetworkInference/Private/NeuralOperators/CoshOperator.cpp
+++ b/Engine/Plugins/Experimental/NeuralNetworkInference/Source/NeuralNetworkInference/Private/NeuralOperators/CoshOperator.cpp
@@ -1,6 +1,7 @@
 // Copyright Epic Games, Inc. All Rights Reserved.
 
 #include "NeuralOperators/CoshOperator.h"
+#include "NeuralOperatorEnumClasses.h"
 #include <cmath>
 
 
@@ -9,7 +10,7 @@
  *****************************************************************************/
 
 FCoshOperator::FCoshOperator(const bool bIsInlinedTensor)
-	: IElementWiseOperator(TEXT("Cosh"), 9, EElementWiseOperator::Cosh, bIsInlinedTensor)
+	: IElementWiseOperator(TEXT("Cosh"), 9, MakeShared<EElementWiseOperator>(EElementWiseOperator::Cosh), bIsInlinedTensor)
 {
 }
 
diff --git a/Engine/Plugins/Experimental/NeuralNetworkInference/Source/NeuralNetworkInference/Private/NeuralOperators/DivOperator.cpp b/Engine/Plugins/Experimental/NeuralNetworkInference/Source/NeuralNetworkInference/Private/NeuralOperators/DivOperator.cpp
--- a/Engine/Plugins/Experimental/NeuralNetworkInference/Source/NeuralNetworkInference/Private/NeuralOperators/DivOperator.cpp
+++ b/Engine/Plugins/Experimental/NeuralNetworkInference/Source/NeuralNetworkInference/Private/NeuralOperators/DivOperator.cpp
@@ -1,6 +1,7 @@
 // Copyright Epic Games, Inc. All Rights Reserved.
 
 #include "NeuralOperators/DivOperator.h"
+#include "NeuralOperatorEnumClasses.h"
 
 
 
@@ -8,7 +9,7 @@
  *****************************************************************************/
 
 FDivOperator::FDivOperator(const TSet<uint32>& InPotentialInlinedTensors)
-	: IMultidirectionalBroadcastOperator(TEXT("Div"), 13, EMultidirectionalBroadcastOperator::Div, InPotentialInlinedTensors)
+	: IMultidirectionalBroadcastOperator(TEXT("Div"), 13, MakeShared<EMultidirectionalBroadcastOperator>(EMultidirectionalBroadcastOperator::Div), InPotentialInlinedTensors)
 {
 }
 
diff --git a/Engine/Plugins/Experimental/NeuralNetworkInference/Source/NeuralNetworkInference/Private/NeuralOperators/ElementWiseOperator.cpp b/Engine/Plugins/Experimental/NeuralNetworkInference/Source/NeuralNetworkInference/Private/NeuralOperators/ElementWiseOperator.cpp
--- a/Engine/Plugins/Experimental/NeuralNetworkInference/Source/NeuralNetworkInference/Private/NeuralOperators/ElementWiseOperator.cpp
+++ b/Engine/Plugins/Experimental/NeuralNetworkInference/Source/NeuralNetworkInference/Private/NeuralOperators/ElementWiseOperator.cpp
@@ -4,13 +4,14 @@
 #include "NeuralNetworkInferenceShaders/ElementWiseCS.h"
 #include "NeuralNetworkInferenceUtils.h"
 #include "NeuralNetworkInferenceUtilsGPU.h"
+#include "NeuralOperatorEnumClasses.h"
 
 
 
 /* IElementWiseOperator structors
  *****************************************************************************/
 
-IElementWiseOperator::IElementWiseOperator(const FString& InName, const int32 InVersion, const EElementWiseOperator InElementWiseOperator,
+IElementWiseOperator::IElementWiseOperator(const FString& InName, const int32 InVersion, const TSharedPtr<EElementWiseOperator>& InElementWiseOperator,
 	const bool bIsInlinedTensor, const TArray<float>& InAttributes)
 	: FNeuralOperator(InName, InVersion, (bIsInlinedTensor ? 0 : -1))
 	, Attributes(InAttributes)
@@ -84,7 +85,7 @@ void IElementWiseOperator::ForwardGPU_RenderThread(FRDGBuilder* InOutGraphBuilde
 	Parameters->OutputUAV = (InlinedTensor < 0 ? GetOutputTensorNoConst() : GetInputTensorNoConst()).GetBufferUAVRef(); // Not inlined vs. inlined
 	// Set shader
 	FElementWiseCS::FPermutationDomain PermutationVector;
-	PermutationVector.Set<FElementWiseCS::FShaderType>(ElementWiseOperator);
+	PermutationVector.Set<FElementWiseCS::FShaderType>(*ElementWiseOperator);
 	PermutationVector.Set<FElementWiseCS::FIsInlined>((InlinedTensor > -1));
 	TShaderMapRef<FElementWiseCS> ComputeShader(GetGlobalShaderMap(GMaxRHIFeatureLevel), PermutationVector);
 	// Run shader
diff --git a/Engine/Plugins/Experimental/NeuralNetworkInference/Source/NeuralNetworkInference/Private/NeuralOperators/ExpOperator.cpp b/Engine/Plugins/Experimental/NeuralNetworkInference/Source/NeuralNetworkInference/Private/NeuralOperators/ExpOperator.cpp
--- a/Engine/Plugins/Experimental/NeuralNetworkInference/Source/NeuralNetworkInference/Private/NeuralOperators/ExpOperator.cpp
+++ b/Engine/Plugins/Experimental/NeuralNetworkInference/Source/NeuralNetworkInference/Private/NeuralOperators/ExpOperator.cpp
@@ -1,6 +1,7 @@
 // Copyright Epic Games, Inc. All Rights Reserved.
 
 #include "NeuralOperators/ExpOperator.h"
+#include "NeuralOperatorEnumClasses.h"
 
 
 
@@ -8,7 +9,7 @@
  *****************************************************************************/
 
 FExpOperator::FExpOperator(const bool bIsInlinedTensor)
-	: IElementWiseOperator(TEXT("Exp"), 13, EElementWiseOperator::Exp, bIsInlinedTensor)
+	: IElementWiseOperator(TEXT("Exp"), 13, MakeShared<EElementWiseOperator>(EElementWiseOperator::Exp), bIsInlinedTensor)
 {
 }
 
diff --git a/Engine/Plugins/Experimental/NeuralNetworkInference/Source/NeuralNetworkInference/Private/NeuralOperators/FloorOperator.cpp b/Engine/Plugins/Experimental/NeuralNetworkInference/Source/NeuralNetworkInference/Private/NeuralOperators/FloorOperator.cpp
--- a/Engine/Plugins/Experimental/NeuralNetworkInference/Source/NeuralNetworkInference/Private/NeuralOperators/FloorOperator.cpp
+++ b/Engine/Plugins/Experimental/NeuralNetworkInference/Source/NeuralNetworkInference/Private/NeuralOperators/FloorOperator.cpp
@@ -1,6 +1,7 @@
 // Copyright Epic Games, Inc. All Rights Reserved.
 
 #include "NeuralOperators/FloorOperator.h"
+#include "NeuralOperatorEnumClasses.h"
 
 
 
@@ -8,7 +9,7 @@
  *****************************************************************************/
 
 FFloorOperator::FFloorOperator(const bool bIsInlinedTensor)
-	: IElementWiseOperator(TEXT("Floor"), 13, EElementWiseOperator::Floor, bIsInlinedTensor)
+	: IElementWiseOperator(TEXT("Floor"), 13, MakeShared<EElementWiseOperator>(EElementWiseOperator::Floor), bIsInlinedTensor)
 {
 }
 
diff --git a/Engine/Plugins/Experimental/NeuralNetworkInference/Source/NeuralNetworkInference/Private/NeuralOperators/LeakyReluOperator.cpp b/Engine/Plugins/Experimental/NeuralNetworkInference/Source/NeuralNetworkInference/Private/NeuralOperators/LeakyReluOperator.cpp
--- a/Engine/Plugins/Experimental/NeuralNetworkInference/Source/NeuralNetworkInference/Private/NeuralOperators/LeakyReluOperator.cpp
+++ b/Engine/Plugins/Experimental/NeuralNetworkInference/Source/NeuralNetworkInference/Private/NeuralOperators/LeakyReluOperator.cpp
@@ -3,6 +3,7 @@
 #include "NeuralOperators/LeakyReluOperator.h"
 #include "ModelProto.h"
 #include "NeuralNetworkInferenceUtils.h"
+#include "NeuralOperatorEnumClasses.h"
 
 
 
@@ -26,7 +27,7 @@ FLeakyReluOperator::FLeakyReluOperator(const bool bIsInlinedTensor, const FNodeP
 }
 
 FLeakyReluOperator::FLeakyReluOperator(const bool bIsInlinedTensor, const float InAlpha)
-	: IElementWiseOperator(TEXT("LeakyRelu"), 6, EElementWiseOperator::LeakyRelu, bIsInlinedTensor, { InAlpha })
+	: IElementWiseOperator(TEXT("LeakyRelu"), 6, MakeShared<EElementWiseOperator>(EElementWiseOperator::LeakyRelu), bIsInlinedTensor, { InAlpha })
 {
 }
 
diff --git a/Engine/Plugins/Experimental/NeuralNetworkInference/Source/NeuralNetworkInference/Private/NeuralOperators/LogOperator.cpp b/Engine/Plugins/Experimental/NeuralNetworkInference/Source/NeuralNetworkInference/Private/NeuralOperators/LogOperator.cpp
--- a/Engine/Plugins/Experimental/NeuralNetworkInference/Source/NeuralNetworkInference/Private/NeuralOperators/LogOperator.cpp
+++ b/Engine/Plugins/Experimental/NeuralNetworkInference/Source/NeuralNetworkInference/Private/NeuralOperators/LogOperator.cpp
@@ -1,6 +1,7 @@
 // Copyright Epic Games, Inc. All Rights Reserved.
 
 #include "NeuralOperators/LogOperator.h"
+#include "NeuralOperatorEnumClasses.h"
 
 
 
@@ -8,7 +9,7 @@
  *****************************************************************************/
 
 FLogOperator::FLogOperator(const bool bIsInlinedTensor)
-	: IElementWiseOperator(TEXT("Log"), 13, EElementWiseOperator::Log, bIsInlinedTensor)
+	: IElementWiseOperator(TEXT("Log"), 13, MakeShared<EElementWiseOperator>(EElementWiseOperator::Log), bIsInlinedTensor)
 {
 }
 
diff --git a/Engine/Plugins/Experimental/NeuralNetworkInference/Source/NeuralNetworkInference/Private/NeuralOperators/MulOperator.cpp b/Engine/Plugins/Experimental/NeuralNetworkInference/Source/NeuralNetworkInference/Private/NeuralOperators/MulOperator.cpp
--- a/Engine/Plugins/Experimental/NeuralNetworkInference/Source/NeuralNetworkInference/Private/NeuralOperators/MulOperator.cpp
+++ b/Engine/Plugins/Experimental/NeuralNetworkInference/Source/NeuralNetworkInference/Private/NeuralOperators/MulOperator.cpp
@@ -1,6 +1,7 @@
 // Copyright Epic Games, Inc. All Rights Reserved.
 
 #include "NeuralOperators/MulOperator.h"
+#include "NeuralOperatorEnumClasses.h"
 
 
 
@@ -8,7 +9,7 @@
  *****************************************************************************/
 
 FMulOperator::FMulOperator(const TSet<uint32>& InPotentialInlinedTensors)
-	: IMultidirectionalBroadcastOperator(TEXT("Mul"), 13, EMultidirectionalBroadcastOperator::Mul, InPotentialInlinedTensors)
+	: IMultidirectionalBroadcastOperator(TEXT("Mul"), 13, MakeShared<EMultidirectionalBroadcastOperator>(EMultidirectionalBroadcastOperator::Mul), InPotentialInlinedTensors)
 {
 }
 
diff --git a/Engine/Plugins/Experimental/NeuralNetworkInference/Source/NeuralNetworkInference/Private/NeuralOperators/MultidirectionalBroadcastOperator.cpp b/Engine/Plugins/Experimental/NeuralNetworkInference/Source/NeuralNetworkInference/Private/NeuralOperators/MultidirectionalBroadcastOperator.cpp
--- a/Engine/Plugins/Experimental/NeuralNetworkInference/Source/NeuralNetworkInference/Private/NeuralOperators/MultidirectionalBroadcastOperator.cpp
+++ b/Engine/Plugins/Experimental/NeuralNetworkInference/Source/NeuralNetworkInference/Private/NeuralOperators/MultidirectionalBroadcastOperator.cpp
@@ -4,12 +4,12 @@
 #include "NeuralNetworkInferenceShaders/MultidirectionalBroadcastCS.h"
 #include "NeuralNetworkInferenceUtils.h"
 #include "NeuralNetworkInferenceUtilsGPU.h"
+#include "NeuralOperatorEnumClasses.h"
 #include "RHI.h"
 #include "Shader.h" // TShaderRef
 
 
 
-
 /* FPrivateMultidirectionalBroadcastOperator static functions
  *****************************************************************************/
 
@@ -173,10 +173,10 @@ EMultidirectionalBroadcastShapeMode FPrivateMultidirectionalBroadcastOperator::G
  *****************************************************************************/
 
 IMultidirectionalBroadcastOperator::IMultidirectionalBroadcastOperator(const FString& InName, const int32 InVersion,
-	const EMultidirectionalBroadcastOperator InMultidirectionalBroadcastOperator, const TSet<uint32>& InPotentialInlinedTensors)
+	const TSharedPtr<EMultidirectionalBroadcastOperator>& InMultidirectionalBroadcastOperator, const TSet<uint32>& InPotentialInlinedTensors)
 	: FNeuralOperator(InName, InVersion, InPotentialInlinedTensors)
 	, MultidirectionalBroadcastOperator(InMultidirectionalBroadcastOperator)
-	, MultidirectionalBroadcastShapeMode(EMultidirectionalBroadcastShapeMode::MAX)
+	, MultidirectionalBroadcastShapeMode(MakeShared<EMultidirectionalBroadcastShapeMode>(EMultidirectionalBroadcastShapeMode::MAX))
 {
 }
 
@@ -245,9 +245,9 @@ void IMultidirectionalBroadcastOperator::ToGPU_RenderThread()
 	}
 	// Set MultidirectionalBroadcastShapeMode
 	const TArray<FNeuralTensor*>& InputTensors = GetInputTensorsConst();
-	MultidirectionalBroadcastShapeMode = FPrivateMultidirectionalBroadcastOperator::GetMultidirectionalBroadcastShapeMode(*InputTensors[0], *InputTensors[1]);
+	*MultidirectionalBroadcastShapeMode = FPrivateMultidirectionalBroadcastOperator::GetMultidirectionalBroadcastShapeMode(*InputTensors[0], *InputTensors[1]);
 	// Memory to GPU
-	if (MultidirectionalBroadcastShapeMode == EMultidirectionalBroadcastShapeMode::MultidirectionalBroadcast)
+	if (*MultidirectionalBroadcastShapeMode == EMultidirectionalBroadcastShapeMode::MultidirectionalBroadcast)
 	{
 		ShapesX.CreateAndLoadSRVBuffer(TEXT("MultidirectionalBroadcast_ShapesX"));
 		ShapesY.CreateAndLoadSRVBuffer(TEXT("MultidirectionalBroadcast_ShapesY"));
@@ -280,7 +280,7 @@ void IMultidirectionalBroadcastOperator::ForwardGPU_RenderThread(FRDGBuilder* In
 		Parameters->ASRV = X.GetBufferSRVRef();
 		Parameters->BSRV = Y.GetBufferSRVRef();
 		Parameters->OutputUAV = GetOutputTensorNoConst().GetBufferUAVRef();
-		if (MultidirectionalBroadcastShapeMode == EMultidirectionalBroadcastShapeMode::MultidirectionalBroadcast)
+		if (*MultidirectionalBroadcastShapeMode == EMultidirectionalBroadcastShapeMode::MultidirectionalBroadcast)
 		{
 			Parameters->ShapeDimensions = ShapesOutput.Num();
 			Parameters->ShapeOutput = ShapesOutput.GetUInt32SRV();
@@ -298,7 +298,7 @@ void IMultidirectionalBroadcastOperator::ForwardGPU_RenderThread(FRDGBuilder* In
 			MultidirectionalBroadcastInlinedMode = EMultidirectionalBroadcastInlinedMode::A;
 			Parameters->BSRV = Y.GetBufferSRVRef();
 			Parameters->OutputUAV = InputTensorsNotConst[0]->GetBufferUAVRef();
-			if (MultidirectionalBroadcastShapeMode == EMultidirectionalBroadcastShapeMode::MultidirectionalBroadcast)
+			if (*MultidirectionalBroadcastShapeMode == EMultidirectionalBroadcastShapeMode::MultidirectionalBroadcast)
 			{
 				Parameters->ShapeDimensions = ShapesOutput.Num();
 				Parameters->ShapeOutput = ShapesOutput.GetUInt32SRV();
@@ -311,7 +311,7 @@ void IMultidirectionalBroadcastOperator::ForwardGPU_RenderThread(FRDGBuilder* In
 			MultidirectionalBroadcastInlinedMode = EMultidirectionalBroadcastInlinedMode::B;
 			Parameters->ASRV = X.GetBufferSRVRef();
 			Parameters->OutputUAV = InputTensorsNotConst[1]->GetBufferUAVRef();
-			if (MultidirectionalBroadcastShapeMode == EMultidirectionalBroadcastShapeMode::MultidirectionalBroadcast)
+			if (*MultidirectionalBroadcastShapeMode == EMultidirectionalBroadcastShapeMode::MultidirectionalBroadcast)
 			{
 				Parameters->ShapeDimensions = ShapesOutput.Num();
 				Parameters->ShapeOutput = ShapesOutput.GetUInt32SRV();
@@ -328,9 +328,9 @@ void IMultidirectionalBroadcastOperator::ForwardGPU_RenderThread(FRDGBuilder* In
 	}
 	// Set shader
 	FMultidirectionalBroadcastCS::FPermutationDomain PermutationVector;
-	PermutationVector.Set<FMultidirectionalBroadcastCS::FShaderType>(MultidirectionalBroadcastOperator);
+	PermutationVector.Set<FMultidirectionalBroadcastCS::FShaderType>(*MultidirectionalBroadcastOperator);
 	PermutationVector.Set<FMultidirectionalBroadcastCS::FInlinedMode>(MultidirectionalBroadcastInlinedMode);
-	PermutationVector.Set<FMultidirectionalBroadcastCS::FShapeMode>(MultidirectionalBroadcastShapeMode);
+	PermutationVector.Set<FMultidirectionalBroadcastCS::FShapeMode>(*MultidirectionalBroadcastShapeMode);
 	TShaderMapRef<FMultidirectionalBroadcastCS> ComputeShader(GetGlobalShaderMap(GMaxRHIFeatureLevel), PermutationVector);
 	// Run shader
 	const uint32 ThreadGroupCountValueX = FMath::DivideAndRoundUp(Volume, FMultidirectionalBroadcastCS::THREADGROUP_SIZE_X);
diff --git a/Engine/Plugins/Experimental/NeuralNetworkInference/Source/NeuralNetworkInference/Private/NeuralOperators/NegOperator.cpp b/Engine/Plugins/Experimental/NeuralNetworkInference/Source/NeuralNetworkInference/Private/NeuralOperators/NegOperator.cpp
--- a/Engine/Plugins/Experimental/NeuralNetworkInference/Source/NeuralNetworkInference/Private/NeuralOperators/NegOperator.cpp
+++ b/Engine/Plugins/Experimental/NeuralNetworkInference/Source/NeuralNetworkInference/Private/NeuralOperators/NegOperator.cpp
@@ -1,6 +1,7 @@
 // Copyright Epic Games, Inc. All Rights Reserved.
 
 #include "NeuralOperators/NegOperator.h"
+#include "NeuralOperatorEnumClasses.h"
 
 
 
@@ -8,7 +9,7 @@
  *****************************************************************************/
 
 FNegOperator::FNegOperator(const bool bIsInlinedTensor)
-	: IElementWiseOperator(TEXT("Neg"), 13, EElementWiseOperator::Neg, bIsInlinedTensor)
+	: IElementWiseOperator(TEXT("Neg"), 13, MakeShared<EElementWiseOperator>(EElementWiseOperator::Neg), bIsInlinedTensor)
 {
 }
 
diff --git a/Engine/Plugins/Experimental/NeuralNetworkInference/Source/NeuralNetworkInference/Private/NeuralOperators/PowOperator.cpp b/Engine/Plugins/Experimental/NeuralNetworkInference/Source/NeuralNetworkInference/Private/NeuralOperators/PowOperator.cpp
--- a/Engine/Plugins/Experimental/NeuralNetworkInference/Source/NeuralNetworkInference/Private/NeuralOperators/PowOperator.cpp
+++ b/Engine/Plugins/Experimental/NeuralNetworkInference/Source/NeuralNetworkInference/Private/NeuralOperators/PowOperator.cpp
@@ -1,6 +1,7 @@
 // Copyright Epic Games, Inc. All Rights Reserved.
 
 #include "NeuralOperators/PowOperator.h"
+#include "NeuralOperatorEnumClasses.h"
 
 
 
@@ -8,7 +9,7 @@
  *****************************************************************************/
 
 FPowOperator::FPowOperator(const TSet<uint32>& InPotentialInlinedTensors)
-	: IMultidirectionalBroadcastOperator(TEXT("Pow"), 13, EMultidirectionalBroadcastOperator::Pow, InPotentialInlinedTensors)
+	: IMultidirectionalBroadcastOperator(TEXT("Pow"), 13, MakeShared<EMultidirectionalBroadcastOperator>(EMultidirectionalBroadcastOperator::Pow), InPotentialInlinedTensors)
 {
 }
 
diff --git a/Engine/Plugins/Experimental/NeuralNetworkInference/Source/NeuralNetworkInference/Private/NeuralOperators/ReciprocalOperator.cpp b/Engine/Plugins/Experimental/NeuralNetworkInference/Source/NeuralNetworkInference/Private/NeuralOperators/ReciprocalOperator.cpp
--- a/Engine/Plugins/Experimental/NeuralNetworkInference/Source/NeuralNetworkInference/Private/NeuralOperators/ReciprocalOperator.cpp
+++ b/Engine/Plugins/Experimental/NeuralNetworkInference/Source/NeuralNetworkInference/Private/NeuralOperators/ReciprocalOperator.cpp
@@ -1,6 +1,7 @@
 // Copyright Epic Games, Inc. All Rights Reserved.
 
 #include "NeuralOperators/ReciprocalOperator.h"
+#include "NeuralOperatorEnumClasses.h"
 
 
 
@@ -8,7 +9,7 @@
  *****************************************************************************/
 
 FReciprocalOperator::FReciprocalOperator(const bool bIsInlinedTensor)
-	: IElementWiseOperator(TEXT("Reciprocal"), 13, EElementWiseOperator::Reciprocal, bIsInlinedTensor)
+	: IElementWiseOperator(TEXT("Reciprocal"), 13, MakeShared<EElementWiseOperator>(EElementWiseOperator::Reciprocal), bIsInlinedTensor)
 {
 }
 
diff --git a/Engine/Plugins/Experimental/NeuralNetworkInference/Source/NeuralNetworkInference/Private/NeuralOperators/ReluOperator.cpp b/Engine/Plugins/Experimental/NeuralNetworkInference/Source/NeuralNetworkInference/Private/NeuralOperators/ReluOperator.cpp
--- a/Engine/Plugins/Experimental/NeuralNetworkInference/Source/NeuralNetworkInference/Private/NeuralOperators/ReluOperator.cpp
+++ b/Engine/Plugins/Experimental/NeuralNetworkInference/Source/NeuralNetworkInference/Private/NeuralOperators/ReluOperator.cpp
@@ -1,6 +1,7 @@
 // Copyright Epic Games, Inc. All Rights Reserved.
 
 #include "NeuralOperators/ReluOperator.h"
+#include "NeuralOperatorEnumClasses.h"
 
 
 
@@ -8,7 +9,7 @@
  *****************************************************************************/
 
 FReluOperator::FReluOperator(const bool bIsInlinedTensor)
-	: IElementWiseOperator(TEXT("Relu"), 13, EElementWiseOperator::Relu, bIsInlinedTensor)
+	: IElementWiseOperator(TEXT("Relu"), 13, MakeShared<EElementWiseOperator>(EElementWiseOperator::Relu), bIsInlinedTensor)
 {
 }
 
diff --git a/Engine/Plugins/Experimental/NeuralNetworkInference/Source/NeuralNetworkInference/Private/NeuralOperators/RoundOperator.cpp b/Engine/Plugins/Experimental/NeuralNetworkInference/Source/NeuralNetworkInference/Private/NeuralOperators/RoundOperator.cpp
--- a/Engine/Plugins/Experimental/NeuralNetworkInference/Source/NeuralNetworkInference/Private/NeuralOperators/RoundOperator.cpp
+++ b/Engine/Plugins/Experimental/NeuralNetworkInference/Source/NeuralNetworkInference/Private/NeuralOperators/RoundOperator.cpp
@@ -1,6 +1,7 @@
 // Copyright Epic Games, Inc. All Rights Reserved.
 
 #include "NeuralOperators/RoundOperator.h"
+#include "NeuralOperatorEnumClasses.h"
 
 
 
@@ -8,7 +9,7 @@
  *****************************************************************************/
 
 FRoundOperator::FRoundOperator(const bool bIsInlinedTensor)
-	: IElementWiseOperator(TEXT("Round"), 11, EElementWiseOperator::Round, bIsInlinedTensor)
+	: IElementWiseOperator(TEXT("Round"), 11, MakeShared<EElementWiseOperator>(EElementWiseOperator::Round), bIsInlinedTensor)
 {
 }
 
diff --git a/Engine/Plugins/Experimental/NeuralNetworkInference/Source/NeuralNetworkInference/Private/NeuralOperators/SigmoidOperator.cpp b/Engine/Plugins/Experimental/NeuralNetworkInference/Source/NeuralNetworkInference/Private/NeuralOperators/SigmoidOperator.cpp
--- a/Engine/Plugins/Experimental/NeuralNetworkInference/Source/NeuralNetworkInference/Private/NeuralOperators/SigmoidOperator.cpp
+++ b/Engine/Plugins/Experimental/NeuralNetworkInference/Source/NeuralNetworkInference/Private/NeuralOperators/SigmoidOperator.cpp
@@ -1,6 +1,7 @@
 // Copyright Epic Games, Inc. All Rights Reserved.
 
 #include "NeuralOperators/SigmoidOperator.h"
+#include "NeuralOperatorEnumClasses.h"
 
 
 
@@ -8,7 +9,7 @@
  *****************************************************************************/
 
 FSigmoidOperator::FSigmoidOperator(const bool bIsInlinedTensor)
-	: IElementWiseOperator(TEXT("Sigmoid"), 13, EElementWiseOperator::Sigmoid, bIsInlinedTensor)
+	: IElementWiseOperator(TEXT("Sigmoid"), 13, MakeShared<EElementWiseOperator>(EElementWiseOperator::Sigmoid), bIsInlinedTensor)
 {
 }
 
diff --git a/Engine/Plugins/Experimental/NeuralNetworkInference/Source/NeuralNetworkInference/Private/NeuralOperators/SignOperator.cpp b/Engine/Plugins/Experimental/NeuralNetworkInference/Source/NeuralNetworkInference/Private/NeuralOperators/SignOperator.cpp
--- a/Engine/Plugins/Experimental/NeuralNetworkInference/Source/NeuralNetworkInference/Private/NeuralOperators/SignOperator.cpp
+++ b/Engine/Plugins/Experimental/NeuralNetworkInference/Source/NeuralNetworkInference/Private/NeuralOperators/SignOperator.cpp
@@ -1,6 +1,7 @@
 // Copyright Epic Games, Inc. All Rights Reserved.
 
 #include "NeuralOperators/SignOperator.h"
+#include "NeuralOperatorEnumClasses.h"
 
 
 
@@ -8,7 +9,7 @@
  *****************************************************************************/
 
 FSignOperator::FSignOperator(const bool bIsInlinedTensor)
-	: IElementWiseOperator(TEXT("Sign"), 13, EElementWiseOperator::Sign, bIsInlinedTensor)
+	: IElementWiseOperator(TEXT("Sign"), 13, MakeShared<EElementWiseOperator>(EElementWiseOperator::Sign), bIsInlinedTensor)
 {
 }
 
diff --git a/Engine/Plugins/Experimental/NeuralNetworkInference/Source/NeuralNetworkInference/Private/NeuralOperators/SinOperator.cpp b/Engine/Plugins/Experimental/NeuralNetworkInference/Source/NeuralNetworkInference/Private/NeuralOperators/SinOperator.cpp
--- a/Engine/Plugins/Experimental/NeuralNetworkInference/Source/NeuralNetworkInference/Private/NeuralOperators/SinOperator.cpp
+++ b/Engine/Plugins/Experimental/NeuralNetworkInference/Source/NeuralNetworkInference/Private/NeuralOperators/SinOperator.cpp
@@ -1,6 +1,7 @@
 // Copyright Epic Games, Inc. All Rights Reserved.
 
 #include "NeuralOperators/SinOperator.h"
+#include "NeuralOperatorEnumClasses.h"
 
 
 
@@ -8,7 +9,7 @@
  *****************************************************************************/
 
 FSinOperator::FSinOperator(const bool bIsInlinedTensor)
-	: IElementWiseOperator(TEXT("Sin"), 7, EElementWiseOperator::Sin, bIsInlinedTensor)
+	: IElementWiseOperator(TEXT("Sin"), 7, MakeShared<EElementWiseOperator>(EElementWiseOperator::Sin), bIsInlinedTensor)
 {
 }
 
diff --git a/Engine/Plugins/Experimental/NeuralNetworkInference/Source/NeuralNetworkInference/Private/NeuralOperators/SinhOperator.cpp b/Engine/Plugins/Experimental/NeuralNetworkInference/Source/NeuralNetworkInference/Private/NeuralOperators/SinhOperator.cpp
--- a/Engine/Plugins/Experimental/NeuralNetworkInference/Source/NeuralNetworkInference/Private/NeuralOperators/SinhOperator.cpp
+++ b/Engine/Plugins/Experimental/NeuralNetworkInference/Source/NeuralNetworkInference/Private/NeuralOperators/SinhOperator.cpp
@@ -1,6 +1,7 @@
 // Copyright Epic Games, Inc. All Rights Reserved.
 
 #include "NeuralOperators/SinhOperator.h"
+#include "NeuralOperatorEnumClasses.h"
 
 
 
@@ -8,7 +9,7 @@
  *****************************************************************************/
 
 FSinhOperator::FSinhOperator(const bool bIsInlinedTensor)
-	: IElementWiseOperator(TEXT("Sinh"), 9, EElementWiseOperator::Sinh, bIsInlinedTensor)
+	: IElementWiseOperator(TEXT("Sinh"), 9, MakeShared<EElementWiseOperator>(EElementWiseOperator::Sinh), bIsInlinedTensor)
 {
 }
 
diff --git a/Engine/Plugins/Experimental/NeuralNetworkInference/Source/NeuralNetworkInference/Private/NeuralOperators/SqrtOperator.cpp b/Engine/Plugins/Experimental/NeuralNetworkInference/Source/NeuralNetworkInference/Private/NeuralOperators/SqrtOperator.cpp
--- a/Engine/Plugins/Experimental/NeuralNetworkInference/Source/NeuralNetworkInference/Private/NeuralOperators/SqrtOperator.cpp
+++ b/Engine/Plugins/Experimental/NeuralNetworkInference/Source/NeuralNetworkInference/Private/NeuralOperators/SqrtOperator.cpp
@@ -1,6 +1,7 @@
 // Copyright Epic Games, Inc. All Rights Reserved.
 
 #include "NeuralOperators/SqrtOperator.h"
+#include "NeuralOperatorEnumClasses.h"
 
 
 
@@ -8,7 +9,7 @@
  *****************************************************************************/
 
 FSqrtOperator::FSqrtOperator(const bool bIsInlinedTensor)
-	: IElementWiseOperator(TEXT("Sqrt"), 13, EElementWiseOperator::Sqrt, bIsInlinedTensor)
+	: IElementWiseOperator(TEXT("Sqrt"), 13, MakeShared<EElementWiseOperator>(EElementWiseOperator::Sqrt), bIsInlinedTensor)
 {
 }
 
diff --git a/Engine/Plugins/Experimental/NeuralNetworkInference/Source/NeuralNetworkInference/Private/NeuralOperators/SubOperator.cpp b/Engine/Plugins/Experimental/NeuralNetworkInference/Source/NeuralNetworkInference/Private/NeuralOperators/SubOperator.cpp
--- a/Engine/Plugins/Experimental/NeuralNetworkInference/Source/NeuralNetworkInference/Private/NeuralOperators/SubOperator.cpp
+++ b/Engine/Plugins/Experimental/NeuralNetworkInference/Source/NeuralNetworkInference/Private/NeuralOperators/SubOperator.cpp
@@ -1,6 +1,7 @@
 // Copyright Epic Games, Inc. All Rights Reserved.
 
 #include "NeuralOperators/SubOperator.h"
+#include "NeuralOperatorEnumClasses.h"
 
 
 
@@ -8,7 +9,7 @@
  *****************************************************************************/
 
 FSubOperator::FSubOperator(const TSet<uint32>& InPotentialInlinedTensors)
-	: IMultidirectionalBroadcastOperator(TEXT("Sub"), 13, EMultidirectionalBroadcastOperator::Sub, InPotentialInlinedTensors)
+	: IMultidirectionalBroadcastOperator(TEXT("Sub"), 13, MakeShared<EMultidirectionalBroadcastOperator>(EMultidirectionalBroadcastOperator::Sub), InPotentialInlinedTensors)
 {
 }
 
diff --git a/Engine/Plugins/Experimental/NeuralNetworkInference/Source/NeuralNetworkInference/Private/NeuralOperators/TanOperator.cpp b/Engine/Plugins/Experimental/NeuralNetworkInference/Source/NeuralNetworkInference/Private/NeuralOperators/TanOperator.cpp
--- a/Engine/Plugins/Experimental/NeuralNetworkInference/Source/NeuralNetworkInference/Private/NeuralOperators/TanOperator.cpp
+++ b/Engine/Plugins/Experimental/NeuralNetworkInference/Source/NeuralNetworkInference/Private/NeuralOperators/TanOperator.cpp
@@ -1,6 +1,7 @@
 // Copyright Epic Games, Inc. All Rights Reserved.
 
 #include "NeuralOperators/TanOperator.h"
+#include "NeuralOperatorEnumClasses.h"
 
 
 
@@ -8,7 +9,7 @@
  *****************************************************************************/
 
 FTanOperator::FTanOperator(const bool bIsInlinedTensor)
-	: IElementWiseOperator(TEXT("Tan"), 7, EElementWiseOperator::Tan, bIsInlinedTensor)
+	: IElementWiseOperator(TEXT("Tan"), 7, MakeShared<EElementWiseOperator>(EElementWiseOperator::Tan), bIsInlinedTensor)
 {
 }
 
diff --git a/Engine/Plugins/Experimental/NeuralNetworkInference/Source/NeuralNetworkInference/Private/NeuralOperators/TanhOperator.cpp b/Engine/Plugins/Experimental/NeuralNetworkInference/Source/NeuralNetworkInference/Private/NeuralOperators/TanhOperator.cpp
--- a/Engine/Plugins/Experimental/NeuralNetworkInference/Source/NeuralNetworkInference/Private/NeuralOperators/TanhOperator.cpp
+++ b/Engine/Plugins/Experimental/NeuralNetworkInference/Source/NeuralNetworkInference/Private/NeuralOperators/TanhOperator.cpp
@@ -1,6 +1,7 @@
 // Copyright Epic Games, Inc. All Rights Reserved.
 
 #include "NeuralOperators/TanhOperator.h"
+#include "NeuralOperatorEnumClasses.h"
 #include <cmath>
 
 
@@ -9,7 +10,7 @@
  *****************************************************************************/
 
 FTanhOperator::FTanhOperator(const bool bIsInlinedTensor)
-	: IElementWiseOperator(TEXT("Tanh"), 13, EElementWiseOperator::Tanh, bIsInlinedTensor)
+	: IElementWiseOperator(TEXT("Tanh"), 13, MakeShared<EElementWiseOperator>(EElementWiseOperator::Tanh), bIsInlinedTensor)
 {
 }
 
diff --git a/Engine/Plugins/Experimental/NeuralNetworkInference/Source/NeuralNetworkInferenceCore/NeuralNetworkInferenceCore.Build.cs b/Engine/Plugins/Experimental/NeuralNetworkInference/Source/NeuralNetworkInferenceCore/NeuralNetworkInferenceCore.Build.cs
deleted file mode 100644
diff --git a/Engine/Plugins/Experimental/NeuralNetworkInference/Source/NeuralNetworkInferenceCore/Private/Module.cpp b/Engine/Plugins/Experimental/NeuralNetworkInference/Source/NeuralNetworkInferenceCore/Private/Module.cpp
deleted file mode 100644
diff --git a/Engine/Plugins/Experimental/NeuralNetworkInference/Source/NeuralNetworkInferenceCore/Private/Module.h b/Engine/Plugins/Experimental/NeuralNetworkInference/Source/NeuralNetworkInferenceCore/Private/Module.h
deleted file mode 100644
diff --git a/Engine/Plugins/Experimental/NeuralNetworkInference/Source/NeuralNetworkInferenceCore/Private/NeuralNetworkInferenceCoreUtils.h b/Engine/Plugins/Experimental/NeuralNetworkInference/Source/NeuralNetworkInferenceCore/Private/NeuralNetworkInferenceCoreUtils.h
deleted file mode 100644
diff --git a/Engine/Plugins/Experimental/NeuralNetworkInference/Source/NeuralNetworkInferenceCore/Public/NeuralOperatorEnumClasses.h b/Engine/Plugins/Experimental/NeuralNetworkInference/Source/NeuralNetworkInferenceCore/Public/NeuralOperatorEnumClasses.h
deleted file mode 100644
diff --git a/Engine/Plugins/Experimental/NeuralNetworkInference/Source/Private/ModelProto/ModelProto.Build.cs b/Engine/Plugins/Experimental/NeuralNetworkInference/Source/Private/ModelProto/ModelProto.Build.cs
--- a/Engine/Plugins/Experimental/NeuralNetworkInference/Source/Private/ModelProto/ModelProto.Build.cs
+++ b/Engine/Plugins/Experimental/NeuralNetworkInference/Source/Private/ModelProto/ModelProto.Build.cs
@@ -19,8 +19,7 @@ public class ModelProto : ModuleRules
 			new string[] {
 				"Core",
 				"CoreUObject",
-				"Engine",
-				"NeuralNetworkInferenceCore"
+				"Engine"
 			}
 		);
 	}
diff --git a/Engine/Plugins/Experimental/NeuralNetworkInference/Source/Private/NeuralNetworkInferenceShaders/Internal/NeuralOperatorEnumClasses.h b/Engine/Plugins/Experimental/NeuralNetworkInference/Source/Private/NeuralNetworkInferenceShaders/Internal/NeuralOperatorEnumClasses.h
new file mode 100644
--- /dev/null
+++ b/Engine/Plugins/Experimental/NeuralNetworkInference/Source/Private/NeuralNetworkInferenceShaders/Internal/NeuralOperatorEnumClasses.h
@@ -0,0 +1,61 @@
+// Copyright Epic Games, Inc. All Rights Reserved.
+
+#pragma once
+
+#include "CoreMinimal.h"
+
+enum class EBatchNormalizationMode : uint8
+{
+	OneDimension,
+	NDimensions,
+	MAX /** Needed for shader (for SHADER_PERMUTATION_ENUM_CLASS) */
+};
+
+enum class EConvMode : uint8
+{
+	D1, /** 1D convolution */
+	D2, /** 2D convolution */
+	D3, /** 3D convolution */
+	DN, /** nD convolution */
+	MAX /** Needed for shader (for SHADER_PERMUTATION_ENUM_CLASS) */
+};
+
+enum class EGemmMode : uint8
+{
+	CTensor,
+	CScalar,
+	MAX /** Needed for shader (for SHADER_PERMUTATION_ENUM_CLASS) */
+};
+
+enum class EElementWiseOperator : uint8
+{
+	Abs, Acos, Asin, Atan, Ceil, Cos, Cosh, Exp, Floor, LeakyRelu, Log, Neg, Reciprocal, Round, Relu, Sigmoid, Sign, Sin, Sinh, Sqrt, Tan, Tanh,
+	MAX /** Needed for shader (for SHADER_PERMUTATION_ENUM_CLASS) */
+};
+
+enum class EMultidirectionalBroadcastOperator : uint8
+{
+	Add,
+	Div,
+	Mul,
+	Pow,
+	Sub,
+	MAX /** Needed for shader (for SHADER_PERMUTATION_ENUM_CLASS) */
+};
+
+enum class EMultidirectionalBroadcastInlinedMode : uint8
+{
+	A,
+	B,
+	None,
+	MAX /** Needed for shader (for SHADER_PERMUTATION_ENUM_CLASS) */
+};
+
+enum class EMultidirectionalBroadcastShapeMode : uint8
+{
+	ElementWise,
+	AScalar,
+	BScalar,
+	MultidirectionalBroadcast,
+	MAX /** Needed for shader (for SHADER_PERMUTATION_ENUM_CLASS) */
+};
diff --git a/Engine/Plugins/Experimental/NeuralNetworkInference/Source/Private/NeuralNetworkInferenceShaders/NeuralNetworkInferenceShaders.Build.cs b/Engine/Plugins/Experimental/NeuralNetworkInference/Source/Private/NeuralNetworkInferenceShaders/NeuralNetworkInferenceShaders.Build.cs
--- a/Engine/Plugins/Experimental/NeuralNetworkInference/Source/Private/NeuralNetworkInferenceShaders/NeuralNetworkInferenceShaders.Build.cs
+++ b/Engine/Plugins/Experimental/NeuralNetworkInference/Source/Private/NeuralNetworkInferenceShaders/NeuralNetworkInferenceShaders.Build.cs
@@ -20,8 +20,7 @@ public class NeuralNetworkInferenceShaders : ModuleRules
 			new string[] {
 				"Core",
 				"CoreUObject",
-				"Engine",
-				"NeuralNetworkInferenceCore"
+				"Engine"
 			}
 		);
 
