Cleaning up NeuralNetworkInferenceShaders, removed unnecesary "NeuralNetworkInferenceShaders" folder in its includes.

_codereview Francisco.Vicente
#rnx

diff --git a/Engine/Plugins/Experimental/NeuralNetworkInference/Source/NeuralNetworkInference/Private/NeuralOperators/BatchNormalizationOperator.cpp b/Engine/Plugins/Experimental/NeuralNetworkInference/Source/NeuralNetworkInference/Private/NeuralOperators/BatchNormalizationOperator.cpp
--- a/Engine/Plugins/Experimental/NeuralNetworkInference/Source/NeuralNetworkInference/Private/NeuralOperators/BatchNormalizationOperator.cpp
+++ b/Engine/Plugins/Experimental/NeuralNetworkInference/Source/NeuralNetworkInference/Private/NeuralOperators/BatchNormalizationOperator.cpp
@@ -1,8 +1,8 @@
 // Copyright Epic Games, Inc. All Rights Reserved.
 
 #include "NeuralOperators/BatchNormalizationOperator.h"
+#include "BatchNormalizationCS.h"
 #include "ModelProto.h"
-#include "NeuralNetworkInferenceShaders/BatchNormalizationCS.h"
 #include "NeuralNetworkInferenceUtils.h"
 #include "NeuralNetworkInferenceUtilsGPU.h"
 
diff --git a/Engine/Plugins/Experimental/NeuralNetworkInference/Source/NeuralNetworkInference/Private/NeuralOperators/ConvBaseOperator.cpp b/Engine/Plugins/Experimental/NeuralNetworkInference/Source/NeuralNetworkInference/Private/NeuralOperators/ConvBaseOperator.cpp
--- a/Engine/Plugins/Experimental/NeuralNetworkInference/Source/NeuralNetworkInference/Private/NeuralOperators/ConvBaseOperator.cpp
+++ b/Engine/Plugins/Experimental/NeuralNetworkInference/Source/NeuralNetworkInference/Private/NeuralOperators/ConvBaseOperator.cpp
@@ -1,7 +1,7 @@
 // Copyright Epic Games, Inc. All Rights Reserved.
 
 #include "NeuralOperators/ConvBaseOperator.h"
-#include "NeuralNetworkInferenceShaders/ConvBaseCS.h"
+#include "ConvBaseCS.h"
 #include "NeuralNetworkInferenceUtils.h"
 #include "NeuralNetworkInferenceUtilsGPU.h"
 
diff --git a/Engine/Plugins/Experimental/NeuralNetworkInference/Source/NeuralNetworkInference/Private/NeuralOperators/ConvTransposeOperator.cpp b/Engine/Plugins/Experimental/NeuralNetworkInference/Source/NeuralNetworkInference/Private/NeuralOperators/ConvTransposeOperator.cpp
--- a/Engine/Plugins/Experimental/NeuralNetworkInference/Source/NeuralNetworkInference/Private/NeuralOperators/ConvTransposeOperator.cpp
+++ b/Engine/Plugins/Experimental/NeuralNetworkInference/Source/NeuralNetworkInference/Private/NeuralOperators/ConvTransposeOperator.cpp
@@ -1,8 +1,8 @@
 // Copyright Epic Games, Inc. All Rights Reserved.
 
 #include "NeuralOperators/ConvTransposeOperator.h"
+#include "ConvTransposeCS.h"
 #include "ModelProto.h"
-#include "NeuralNetworkInferenceShaders/ConvTransposeCS.h"
 #include "NeuralNetworkInferenceUtils.h"
 #include "NeuralNetworkInferenceUtilsGPU.h"
 
diff --git a/Engine/Plugins/Experimental/NeuralNetworkInference/Source/NeuralNetworkInference/Private/NeuralOperators/ElementWiseOperator.cpp b/Engine/Plugins/Experimental/NeuralNetworkInference/Source/NeuralNetworkInference/Private/NeuralOperators/ElementWiseOperator.cpp
--- a/Engine/Plugins/Experimental/NeuralNetworkInference/Source/NeuralNetworkInference/Private/NeuralOperators/ElementWiseOperator.cpp
+++ b/Engine/Plugins/Experimental/NeuralNetworkInference/Source/NeuralNetworkInference/Private/NeuralOperators/ElementWiseOperator.cpp
@@ -1,7 +1,7 @@
 // Copyright Epic Games, Inc. All Rights Reserved.
 
 #include "NeuralOperators/ElementWiseOperator.h"
-#include "NeuralNetworkInferenceShaders/ElementWiseCS.h"
+#include "ElementWiseCS.h"
 #include "NeuralNetworkInferenceUtils.h"
 #include "NeuralNetworkInferenceUtilsGPU.h"
 #include "NeuralOperatorEnumClasses.h"
diff --git a/Engine/Plugins/Experimental/NeuralNetworkInference/Source/NeuralNetworkInference/Private/NeuralOperators/GemmOperator.cpp b/Engine/Plugins/Experimental/NeuralNetworkInference/Source/NeuralNetworkInference/Private/NeuralOperators/GemmOperator.cpp
--- a/Engine/Plugins/Experimental/NeuralNetworkInference/Source/NeuralNetworkInference/Private/NeuralOperators/GemmOperator.cpp
+++ b/Engine/Plugins/Experimental/NeuralNetworkInference/Source/NeuralNetworkInference/Private/NeuralOperators/GemmOperator.cpp
@@ -1,8 +1,8 @@
 // Copyright Epic Games, Inc. All Rights Reserved.
 
 #include "NeuralOperators/GemmOperator.h"
+#include "GemmCS.h"
 #include "ModelProto.h"
-#include "NeuralNetworkInferenceShaders/GemmCS.h"
 #include "NeuralNetworkInferenceUtils.h"
 #include "NeuralNetworkInferenceUtilsGPU.h"
 
diff --git a/Engine/Plugins/Experimental/NeuralNetworkInference/Source/NeuralNetworkInference/Private/NeuralOperators/MultidirectionalBroadcastOperator.cpp b/Engine/Plugins/Experimental/NeuralNetworkInference/Source/NeuralNetworkInference/Private/NeuralOperators/MultidirectionalBroadcastOperator.cpp
--- a/Engine/Plugins/Experimental/NeuralNetworkInference/Source/NeuralNetworkInference/Private/NeuralOperators/MultidirectionalBroadcastOperator.cpp
+++ b/Engine/Plugins/Experimental/NeuralNetworkInference/Source/NeuralNetworkInference/Private/NeuralOperators/MultidirectionalBroadcastOperator.cpp
@@ -1,7 +1,7 @@
 // Copyright Epic Games, Inc. All Rights Reserved.
 
 #include "NeuralOperators/MultidirectionalBroadcastOperator.h"
-#include "NeuralNetworkInferenceShaders/MultidirectionalBroadcastCS.h"
+#include "MultidirectionalBroadcastCS.h"
 #include "NeuralNetworkInferenceUtils.h"
 #include "NeuralNetworkInferenceUtilsGPU.h"
 #include "NeuralOperatorEnumClasses.h"
diff --git a/Engine/Plugins/Experimental/NeuralNetworkInference/Source/NeuralNetworkInference/Private/NeuralOperators/ReshapeOperator.cpp b/Engine/Plugins/Experimental/NeuralNetworkInference/Source/NeuralNetworkInference/Private/NeuralOperators/ReshapeOperator.cpp
--- a/Engine/Plugins/Experimental/NeuralNetworkInference/Source/NeuralNetworkInference/Private/NeuralOperators/ReshapeOperator.cpp
+++ b/Engine/Plugins/Experimental/NeuralNetworkInference/Source/NeuralNetworkInference/Private/NeuralOperators/ReshapeOperator.cpp
@@ -1,8 +1,8 @@
 // Copyright Epic Games, Inc. All Rights Reserved.
 
 #include "NeuralOperators/ReshapeOperator.h"
+#include "CopyCS.h"
 #include "ModelProto.h"
-#include "NeuralNetworkInferenceShaders/CopyCS.h"
 #include "NeuralNetworkInferenceUtils.h"
 #include "NeuralNetworkInferenceUtilsGPU.h"
 
diff --git a/Engine/Plugins/Experimental/NeuralNetworkInference/Source/NeuralNetworkInference/Private/NeuralOperators/SqueezeOperator.cpp b/Engine/Plugins/Experimental/NeuralNetworkInference/Source/NeuralNetworkInference/Private/NeuralOperators/SqueezeOperator.cpp
--- a/Engine/Plugins/Experimental/NeuralNetworkInference/Source/NeuralNetworkInference/Private/NeuralOperators/SqueezeOperator.cpp
+++ b/Engine/Plugins/Experimental/NeuralNetworkInference/Source/NeuralNetworkInference/Private/NeuralOperators/SqueezeOperator.cpp
@@ -1,7 +1,7 @@
 // Copyright Epic Games, Inc. All Rights Reserved.
 
 #include "NeuralOperators/SqueezeOperator.h"
-#include "NeuralNetworkInferenceShaders/CopyCS.h"
+#include "CopyCS.h"
 #include "NeuralNetworkInferenceUtils.h"
 #include "NeuralNetworkInferenceUtilsGPU.h"
 
diff --git a/Engine/Plugins/Experimental/NeuralNetworkInference/Source/Private/NeuralNetworkInferenceShaders/Internal/BatchNormalizationCS.h b/Engine/Plugins/Experimental/NeuralNetworkInference/Source/Private/NeuralNetworkInferenceShaders/Internal/BatchNormalizationCS.h
new file mode 100644
--- /dev/null
+++ b/Engine/Plugins/Experimental/NeuralNetworkInference/Source/Private/NeuralNetworkInferenceShaders/Internal/BatchNormalizationCS.h
@@ -0,0 +1,44 @@
+// Copyright Epic Games, Inc. All Rights Reserved.
+
+#pragma once
+
+#include "NeuralOperatorEnumClasses.h"
+// GPU/RHI/shaders
+#include "GlobalShader.h"
+#include "RHI.h"
+#include "ProfilingDebugging/RealtimeGPUProfiler.h"
+#include "RenderGraphUtils.h"
+#include "ShaderParameterUtils.h"
+
+
+
+class NEURALNETWORKINFERENCESHADERS_API FBatchNormalizationCS : public FGlobalShader
+{
+	DECLARE_SHADER_TYPE(FBatchNormalizationCS, Global);
+	SHADER_USE_PARAMETER_STRUCT(FBatchNormalizationCS, FGlobalShader)
+
+public:
+	class FBatchNormalizationMode : SHADER_PERMUTATION_ENUM_CLASS("BATCH_NORMALIZATION_MODE", EBatchNormalizationMode);
+	class FIsInlined : SHADER_PERMUTATION_BOOL("IS_INLINED");
+	using FPermutationDomain = TShaderPermutationDomain<FBatchNormalizationMode, FIsInlined>;
+
+	static const uint32 THREADGROUP_SIZE_X;
+
+	static void ModifyCompilationEnvironment(const FGlobalShaderPermutationParameters& InParameters, FShaderCompilerEnvironment& OutEnvironment);
+
+	BEGIN_SHADER_PARAMETER_STRUCT(FParameters, )
+		// Input variables
+		SHADER_PARAMETER(uint32, BatchSize)
+		SHADER_PARAMETER(uint32, ChannelSize)
+		SHADER_PARAMETER(uint32, ChannelsVolume)
+		SHADER_PARAMETER(uint32, ImageArea)
+		SHADER_PARAMETER(float, Epsilon)
+		// SRV/UAV variables
+		SHADER_PARAMETER_RDG_BUFFER_SRV(Buffer<float>, XSRV)
+		SHADER_PARAMETER_RDG_BUFFER_SRV(Buffer<float>, Scale)
+		SHADER_PARAMETER_RDG_BUFFER_SRV(Buffer<float>, Bias)
+		SHADER_PARAMETER_RDG_BUFFER_SRV(Buffer<float>, Mean)
+		SHADER_PARAMETER_RDG_BUFFER_SRV(Buffer<float>, Variance)
+		SHADER_PARAMETER_RDG_BUFFER_UAV(RWBuffer<float>, OutputUAV)
+	END_SHADER_PARAMETER_STRUCT()
+};
diff --git a/Engine/Plugins/Experimental/NeuralNetworkInference/Source/Private/NeuralNetworkInferenceShaders/Internal/ConvBaseCS.h b/Engine/Plugins/Experimental/NeuralNetworkInference/Source/Private/NeuralNetworkInferenceShaders/Internal/ConvBaseCS.h
new file mode 100644
--- /dev/null
+++ b/Engine/Plugins/Experimental/NeuralNetworkInference/Source/Private/NeuralNetworkInferenceShaders/Internal/ConvBaseCS.h
@@ -0,0 +1,58 @@
+// Copyright Epic Games, Inc. All Rights Reserved.
+
+#pragma once
+
+#include "NeuralOperatorEnumClasses.h"
+// GPU/RHI/shaders
+#include "GlobalShader.h"
+#include "RHI.h"
+#include "ProfilingDebugging/RealtimeGPUProfiler.h"
+#include "RenderGraphUtils.h"
+#include "ShaderParameterUtils.h"
+
+
+
+class NEURALNETWORKINFERENCESHADERS_API FConvBaseCS : public FGlobalShader
+{
+	DECLARE_SHADER_TYPE(FConvBaseCS, Global);
+	SHADER_USE_PARAMETER_STRUCT(FConvBaseCS, FGlobalShader)
+
+public:
+	class FHasBias : SHADER_PERMUTATION_BOOL("HAS_BIAS");
+	class FConvMode : SHADER_PERMUTATION_ENUM_CLASS("CONV_MODE", EConvMode);
+	using FPermutationDomain = TShaderPermutationDomain<FHasBias, FConvMode>;
+
+	static const uint32 THREADGROUP_SIZE_X;
+	static const uint32 MAX_NUMBER_DIMENSIONS;
+
+	static void ModifyCompilationEnvironment(const FGlobalShaderPermutationParameters& InParameters, FShaderCompilerEnvironment& OutEnvironment);
+
+	BEGIN_SHADER_PARAMETER_STRUCT(FParameters, )
+		// All variables are used in all (1D, 2D, 3D, nD) convolutions
+		// Input variables
+		SHADER_PARAMETER(uint32, bIsTransposed)
+		SHADER_PARAMETER(uint32, Group)
+		SHADER_PARAMETER(uint32, MIntoGroup)
+		SHADER_PARAMETER(uint32, CIntoGroup)
+		SHADER_PARAMETER(uint32, OutputVolume)
+		SHADER_PARAMETER(uint32, OutputBatchVolume)
+		SHADER_PARAMETER(uint32, OutputImageArea)
+		SHADER_PARAMETER(uint32, WBatchVolume)
+		SHADER_PARAMETER(uint32, WImageArea)
+		SHADER_PARAMETER(uint32, XOrXWithZerosBatchVolume)
+		SHADER_PARAMETER(uint32, XOrXWithZerosImageArea)
+		SHADER_PARAMETER(int32, NumberConvolutionalDimensions) // Optional - Only for nD convolution
+		// Input SRV variables
+		SHADER_PARAMETER_SRV(Buffer<uint>, OutputSizes)
+		SHADER_PARAMETER_SRV(Buffer<uint>, XOrXWithZerosSizes)
+		SHADER_PARAMETER_SRV(Buffer<uint>, WSizes)
+		SHADER_PARAMETER_SRV(Buffer<uint>, Dilations)
+		SHADER_PARAMETER_SRV(Buffer<uint>, Strides)
+		SHADER_PARAMETER_SRV(Buffer<uint>, Pads)
+		// SRV/UAV variables
+		SHADER_PARAMETER_RDG_BUFFER_SRV(Buffer<float>, XOrXWithZerosSRV)
+		SHADER_PARAMETER_RDG_BUFFER_SRV(Buffer<float>, WSRV)
+		SHADER_PARAMETER_RDG_BUFFER_SRV(Buffer<float>, BSRV) // Optional - Only if bias
+		SHADER_PARAMETER_RDG_BUFFER_UAV(RWBuffer<float>, OutputUAV)
+	END_SHADER_PARAMETER_STRUCT()
+};
diff --git a/Engine/Plugins/Experimental/NeuralNetworkInference/Source/Private/NeuralNetworkInferenceShaders/Internal/ConvTransposeCS.h b/Engine/Plugins/Experimental/NeuralNetworkInference/Source/Private/NeuralNetworkInferenceShaders/Internal/ConvTransposeCS.h
new file mode 100644
--- /dev/null
+++ b/Engine/Plugins/Experimental/NeuralNetworkInference/Source/Private/NeuralNetworkInferenceShaders/Internal/ConvTransposeCS.h
@@ -0,0 +1,41 @@
+// Copyright Epic Games, Inc. All Rights Reserved.
+
+#pragma once
+
+#include "NeuralOperatorEnumClasses.h"
+// GPU/RHI/shaders
+#include "GlobalShader.h"
+#include "RHI.h"
+#include "ProfilingDebugging/RealtimeGPUProfiler.h"
+#include "RenderGraphUtils.h"
+#include "ShaderParameterUtils.h"
+
+
+
+class NEURALNETWORKINFERENCESHADERS_API FConvTransposeCS : public FGlobalShader
+{
+	DECLARE_SHADER_TYPE(FConvTransposeCS, Global);
+	SHADER_USE_PARAMETER_STRUCT(FConvTransposeCS, FGlobalShader)
+
+public:
+	class FConvMode : SHADER_PERMUTATION_ENUM_CLASS("CONV_MODE", EConvMode);
+	using FPermutationDomain = TShaderPermutationDomain<FConvMode>;
+
+	static const uint32 THREADGROUP_SIZE_X;
+
+	static void ModifyCompilationEnvironment(const FGlobalShaderPermutationParameters& InParameters, FShaderCompilerEnvironment& OutEnvironment);
+
+	BEGIN_SHADER_PARAMETER_STRUCT(FParameters, )
+		// All variables are used in all (1D, 2D, 3D, nD) convolutions
+		// Input variables
+		SHADER_PARAMETER(uint32, XVolume)
+		SHADER_PARAMETER(int32, NumberConvolutionalDimensions) // Optional - Only for nD convolution
+		// Input SRV variables
+		SHADER_PARAMETER_SRV(Buffer<uint>, Zeros)
+		SHADER_PARAMETER_SRV(Buffer<uint>, XSizes)
+		SHADER_PARAMETER_SRV(Buffer<uint>, XWithZerosSizes)
+		// SRV/UAV variables
+		SHADER_PARAMETER_RDG_BUFFER_SRV(Buffer<float>, XSRV)
+		SHADER_PARAMETER_RDG_BUFFER_UAV(RWBuffer<float>, XWithZerosUAV)
+	END_SHADER_PARAMETER_STRUCT()
+};
diff --git a/Engine/Plugins/Experimental/NeuralNetworkInference/Source/Private/NeuralNetworkInferenceShaders/Internal/CopyCS.h b/Engine/Plugins/Experimental/NeuralNetworkInference/Source/Private/NeuralNetworkInferenceShaders/Internal/CopyCS.h
new file mode 100644
--- /dev/null
+++ b/Engine/Plugins/Experimental/NeuralNetworkInference/Source/Private/NeuralNetworkInferenceShaders/Internal/CopyCS.h
@@ -0,0 +1,32 @@
+// Copyright Epic Games, Inc. All Rights Reserved.
+
+#pragma once
+
+#include "NeuralOperatorEnumClasses.h"
+// GPU/RHI/shaders
+#include "GlobalShader.h"
+#include "RHI.h"
+#include "ProfilingDebugging/RealtimeGPUProfiler.h"
+#include "RenderGraphUtils.h"
+#include "ShaderParameterUtils.h"
+
+
+
+class NEURALNETWORKINFERENCESHADERS_API FCopyCS : public FGlobalShader
+{
+	DECLARE_SHADER_TYPE(FCopyCS, Global);
+	SHADER_USE_PARAMETER_STRUCT(FCopyCS, FGlobalShader)
+
+public:
+	static const uint32 THREADGROUP_SIZE_X;
+
+	static void ModifyCompilationEnvironment(const FGlobalShaderPermutationParameters& InParameters, FShaderCompilerEnvironment& OutEnvironment);
+
+	BEGIN_SHADER_PARAMETER_STRUCT(FParameters, )
+		// Input variables
+		SHADER_PARAMETER(uint32, Num)
+		// SRV/UAV variables
+		SHADER_PARAMETER_RDG_BUFFER_SRV(Buffer<float>, InputSRV)
+		SHADER_PARAMETER_RDG_BUFFER_UAV(RWBuffer<float>, OutputUAV)
+	END_SHADER_PARAMETER_STRUCT()
+};
diff --git a/Engine/Plugins/Experimental/NeuralNetworkInference/Source/Private/NeuralNetworkInferenceShaders/Internal/ElementWiseCS.h b/Engine/Plugins/Experimental/NeuralNetworkInference/Source/Private/NeuralNetworkInferenceShaders/Internal/ElementWiseCS.h
new file mode 100644
--- /dev/null
+++ b/Engine/Plugins/Experimental/NeuralNetworkInference/Source/Private/NeuralNetworkInferenceShaders/Internal/ElementWiseCS.h
@@ -0,0 +1,43 @@
+// Copyright Epic Games, Inc. All Rights Reserved.
+
+#pragma once
+
+#include "NeuralOperatorEnumClasses.h"
+// GPU/RHI/shaders
+#include "GlobalShader.h"
+#include "RHI.h"
+#include "ProfilingDebugging/RealtimeGPUProfiler.h"
+#include "RenderGraphUtils.h"
+#include "ShaderParameterUtils.h"
+
+
+
+class NEURALNETWORKINFERENCESHADERS_API FElementWiseCS : public FGlobalShader
+{
+	DECLARE_SHADER_TYPE(FElementWiseCS, Global);
+	SHADER_USE_PARAMETER_STRUCT(FElementWiseCS, FGlobalShader)
+
+public:
+	static const uint32 THREADGROUP_SIZE_X;
+
+	class FShaderType : SHADER_PERMUTATION_ENUM_CLASS("SHADER_FUNCTION", EElementWiseOperator);
+	class FIsInlined : SHADER_PERMUTATION_BOOL("IS_INLINED");
+	using FPermutationDomain = TShaderPermutationDomain<FShaderType, FIsInlined>;
+
+	/** Set desired multidirectional operator. */
+	static void ModifyCompilationEnvironment(const FGlobalShaderPermutationParameters& InParameters, FShaderCompilerEnvironment& OutEnvironment);
+
+	BEGIN_SHADER_PARAMETER_STRUCT(FParameters, )
+		// Input variables
+		SHADER_PARAMETER(uint32, TensorSize)
+		SHADER_PARAMETER(float, Attribute)
+		// SRV/UAV variables
+		SHADER_PARAMETER_RDG_BUFFER_UAV(RWBuffer<float>, OutputUAV)
+		// Optional SRV/UAV variables
+		SHADER_PARAMETER_RDG_BUFFER_SRV(Buffer<float>, InputSRV)
+	END_SHADER_PARAMETER_STRUCT()
+
+private:
+	static FString GetFunctionString(const EElementWiseOperator InType);
+	static uint32 GetNumberAttributes(const EElementWiseOperator InType);
+};
diff --git a/Engine/Plugins/Experimental/NeuralNetworkInference/Source/Private/NeuralNetworkInferenceShaders/Internal/GemmCS.h b/Engine/Plugins/Experimental/NeuralNetworkInference/Source/Private/NeuralNetworkInferenceShaders/Internal/GemmCS.h
new file mode 100644
--- /dev/null
+++ b/Engine/Plugins/Experimental/NeuralNetworkInference/Source/Private/NeuralNetworkInferenceShaders/Internal/GemmCS.h
@@ -0,0 +1,50 @@
+// Copyright Epic Games, Inc. All Rights Reserved.
+
+#pragma once
+
+#include "NeuralOperatorEnumClasses.h"
+// GPU/RHI/shaders
+#include "GlobalShader.h"
+#include "RHI.h"
+#include "ProfilingDebugging/RealtimeGPUProfiler.h"
+#include "RenderGraphUtils.h"
+#include "ShaderParameterUtils.h"
+
+
+
+class NEURALNETWORKINFERENCESHADERS_API FGemmCS : public FGlobalShader
+{
+	DECLARE_SHADER_TYPE(FGemmCS, Global);
+	SHADER_USE_PARAMETER_STRUCT(FGemmCS, FGlobalShader)
+
+public:
+	class FGemmMode : SHADER_PERMUTATION_ENUM_CLASS("GEMM_MODE", EGemmMode);
+	using FPermutationDomain = TShaderPermutationDomain<FGemmMode>;
+
+	static const uint32 THREADGROUP_SIZE_X;
+	static const uint32 THREADGROUP_SIZE_Y;
+
+	static void ModifyCompilationEnvironment(const FGlobalShaderPermutationParameters& InParameters, FShaderCompilerEnvironment& OutEnvironment);
+
+	BEGIN_SHADER_PARAMETER_STRUCT(FParameters, )
+		// Input variables
+		SHADER_PARAMETER(float, Alpha)
+		SHADER_PARAMETER(float, Beta)
+		SHADER_PARAMETER(uint32, OutputRows)
+		SHADER_PARAMETER(uint32, OutputColumns)
+		SHADER_PARAMETER(uint32, AColsOrBRows)
+		SHADER_PARAMETER(uint32, AStrideX)
+		SHADER_PARAMETER(uint32, AStrideY)
+		SHADER_PARAMETER(uint32, BStrideX)
+		SHADER_PARAMETER(uint32, BStrideY)
+		SHADER_PARAMETER(uint32, CSizeX)
+		SHADER_PARAMETER(uint32, CSizeY)
+		SHADER_PARAMETER(uint32, OutputStride)
+		// SRV/UAV variables
+		SHADER_PARAMETER_RDG_BUFFER_SRV(Buffer<float>, ASRV)
+		SHADER_PARAMETER_RDG_BUFFER_SRV(Buffer<float>, BSRV)
+		SHADER_PARAMETER_RDG_BUFFER_SRV(Buffer<float>, CSRV)
+		SHADER_PARAMETER(float, BetaTimesCScalar)
+		SHADER_PARAMETER_RDG_BUFFER_UAV(RWBuffer<float>, OutputUAV)
+	END_SHADER_PARAMETER_STRUCT()
+};
diff --git a/Engine/Plugins/Experimental/NeuralNetworkInference/Source/Private/NeuralNetworkInferenceShaders/Internal/MultidirectionalBroadcastCS.h b/Engine/Plugins/Experimental/NeuralNetworkInference/Source/Private/NeuralNetworkInferenceShaders/Internal/MultidirectionalBroadcastCS.h
new file mode 100644
--- /dev/null
+++ b/Engine/Plugins/Experimental/NeuralNetworkInference/Source/Private/NeuralNetworkInferenceShaders/Internal/MultidirectionalBroadcastCS.h
@@ -0,0 +1,48 @@
+// Copyright Epic Games, Inc. All Rights Reserved.
+
+#pragma once
+
+#include "NeuralOperatorEnumClasses.h"
+// GPU/RHI/shaders
+#include "GlobalShader.h"
+#include "RHI.h"
+#include "ProfilingDebugging/RealtimeGPUProfiler.h"
+#include "RenderGraphUtils.h"
+#include "ShaderParameterUtils.h"
+
+
+
+class NEURALNETWORKINFERENCESHADERS_API FMultidirectionalBroadcastCS : public FGlobalShader
+{
+	DECLARE_SHADER_TYPE(FMultidirectionalBroadcastCS, Global);
+	SHADER_USE_PARAMETER_STRUCT(FMultidirectionalBroadcastCS, FGlobalShader)
+
+public:
+	static const uint32 THREADGROUP_SIZE_X;
+	static const uint32 MAX_NUMBER_DIMENSIONS;
+
+	class FShaderType : SHADER_PERMUTATION_ENUM_CLASS("SHADER_FUNCTION", EMultidirectionalBroadcastOperator);
+	class FInlinedMode : SHADER_PERMUTATION_ENUM_CLASS("INLINED_MODE", EMultidirectionalBroadcastInlinedMode);
+	class FShapeMode : SHADER_PERMUTATION_ENUM_CLASS("SHAPE_MODE", EMultidirectionalBroadcastShapeMode);
+	using FPermutationDomain = TShaderPermutationDomain<FShaderType, FInlinedMode, FShapeMode>;
+
+	static void ModifyCompilationEnvironment(const FGlobalShaderPermutationParameters& InParameters, FShaderCompilerEnvironment& OutEnvironment);
+
+	BEGIN_SHADER_PARAMETER_STRUCT(FParameters, )
+		// Input variables
+		SHADER_PARAMETER(uint32, TensorSize)
+		// SRV/UAV variables
+		SHADER_PARAMETER_RDG_BUFFER_UAV(RWBuffer<float>, OutputUAV)
+		// Optional elementwise SRV variables
+		SHADER_PARAMETER_RDG_BUFFER_SRV(Buffer<float>, ASRV)
+		SHADER_PARAMETER_RDG_BUFFER_SRV(Buffer<float>, BSRV)
+		// Optional multidirectional broadcast variables
+		SHADER_PARAMETER(uint32, ShapeDimensions)
+		SHADER_PARAMETER_SRV(Buffer<float>, ShapeOutput)
+		SHADER_PARAMETER_SRV(Buffer<float>, ShapeA)
+		SHADER_PARAMETER_SRV(Buffer<float>, ShapeB)
+	END_SHADER_PARAMETER_STRUCT()
+
+private:
+	static FString GetFunctionString(const EMultidirectionalBroadcastOperator InType);
+};
diff --git a/Engine/Plugins/Experimental/NeuralNetworkInference/Source/Private/NeuralNetworkInferenceShaders/Internal/NeuralNetworkInferenceShaders/BatchNormalizationCS.h b/Engine/Plugins/Experimental/NeuralNetworkInference/Source/Private/NeuralNetworkInferenceShaders/Internal/NeuralNetworkInferenceShaders/BatchNormalizationCS.h
deleted file mode 100644
diff --git a/Engine/Plugins/Experimental/NeuralNetworkInference/Source/Private/NeuralNetworkInferenceShaders/Internal/NeuralNetworkInferenceShaders/ConvBaseCS.h b/Engine/Plugins/Experimental/NeuralNetworkInference/Source/Private/NeuralNetworkInferenceShaders/Internal/NeuralNetworkInferenceShaders/ConvBaseCS.h
deleted file mode 100644
diff --git a/Engine/Plugins/Experimental/NeuralNetworkInference/Source/Private/NeuralNetworkInferenceShaders/Internal/NeuralNetworkInferenceShaders/ConvTransposeCS.h b/Engine/Plugins/Experimental/NeuralNetworkInference/Source/Private/NeuralNetworkInferenceShaders/Internal/NeuralNetworkInferenceShaders/ConvTransposeCS.h
deleted file mode 100644
diff --git a/Engine/Plugins/Experimental/NeuralNetworkInference/Source/Private/NeuralNetworkInferenceShaders/Internal/NeuralNetworkInferenceShaders/CopyCS.h b/Engine/Plugins/Experimental/NeuralNetworkInference/Source/Private/NeuralNetworkInferenceShaders/Internal/NeuralNetworkInferenceShaders/CopyCS.h
deleted file mode 100644
diff --git a/Engine/Plugins/Experimental/NeuralNetworkInference/Source/Private/NeuralNetworkInferenceShaders/Internal/NeuralNetworkInferenceShaders/ElementWiseCS.h b/Engine/Plugins/Experimental/NeuralNetworkInference/Source/Private/NeuralNetworkInferenceShaders/Internal/NeuralNetworkInferenceShaders/ElementWiseCS.h
deleted file mode 100644
diff --git a/Engine/Plugins/Experimental/NeuralNetworkInference/Source/Private/NeuralNetworkInferenceShaders/Internal/NeuralNetworkInferenceShaders/GemmCS.h b/Engine/Plugins/Experimental/NeuralNetworkInference/Source/Private/NeuralNetworkInferenceShaders/Internal/NeuralNetworkInferenceShaders/GemmCS.h
deleted file mode 100644
diff --git a/Engine/Plugins/Experimental/NeuralNetworkInference/Source/Private/NeuralNetworkInferenceShaders/Internal/NeuralNetworkInferenceShaders/MultidirectionalBroadcastCS.h b/Engine/Plugins/Experimental/NeuralNetworkInference/Source/Private/NeuralNetworkInferenceShaders/Internal/NeuralNetworkInferenceShaders/MultidirectionalBroadcastCS.h
deleted file mode 100644
diff --git a/Engine/Plugins/Experimental/NeuralNetworkInference/Source/Private/NeuralNetworkInferenceShaders/Private/BatchNormalizationCS.cpp b/Engine/Plugins/Experimental/NeuralNetworkInference/Source/Private/NeuralNetworkInferenceShaders/Private/BatchNormalizationCS.cpp
new file mode 100644
--- /dev/null
+++ b/Engine/Plugins/Experimental/NeuralNetworkInference/Source/Private/NeuralNetworkInferenceShaders/Private/BatchNormalizationCS.cpp
@@ -0,0 +1,25 @@
+// Copyright Epic Games, Inc. All Rights Reserved.
+
+#include "BatchNormalizationCS.h"
+
+
+
+/* FBatchNormalizationCS public functions
+ *****************************************************************************/
+
+void FBatchNormalizationCS::ModifyCompilationEnvironment(const FGlobalShaderPermutationParameters& InParameters, FShaderCompilerEnvironment& OutEnvironment)
+{
+	FGlobalShader::ModifyCompilationEnvironment(InParameters, OutEnvironment);
+	OutEnvironment.SetDefine(TEXT("THREADGROUP_SIZE_X"), THREADGROUP_SIZE_X);
+}
+
+
+
+const uint32 FBatchNormalizationCS::THREADGROUP_SIZE_X = 128;
+
+
+
+/* Shader implementation
+ *****************************************************************************/
+
+IMPLEMENT_SHADER_TYPE(, FBatchNormalizationCS, TEXT("/Plugins/NeuralNetworkInference/Private/BatchNormalizationOperator.usf"), TEXT("BatchNormalizationCS"), SF_Compute) // Path defined in NeuralNetworkInferenceShadersModule.cpp
diff --git a/Engine/Plugins/Experimental/NeuralNetworkInference/Source/Private/NeuralNetworkInferenceShaders/Private/ConvBaseCS.cpp b/Engine/Plugins/Experimental/NeuralNetworkInference/Source/Private/NeuralNetworkInferenceShaders/Private/ConvBaseCS.cpp
new file mode 100644
--- /dev/null
+++ b/Engine/Plugins/Experimental/NeuralNetworkInference/Source/Private/NeuralNetworkInferenceShaders/Private/ConvBaseCS.cpp
@@ -0,0 +1,30 @@
+// Copyright Epic Games, Inc. All Rights Reserved.
+
+#include "ConvBaseCS.h"
+
+
+
+/* FConvBaseCS static members
+ *****************************************************************************/
+
+const uint32 FConvBaseCS::THREADGROUP_SIZE_X(128);
+const uint32 FConvBaseCS::MAX_NUMBER_DIMENSIONS(16);
+
+
+
+/* FConvBaseCS public functions
+ *****************************************************************************/
+
+void FConvBaseCS::ModifyCompilationEnvironment(const FGlobalShaderPermutationParameters& InParameters, FShaderCompilerEnvironment& OutEnvironment)
+{
+	FGlobalShader::ModifyCompilationEnvironment(InParameters, OutEnvironment);
+	OutEnvironment.SetDefine(TEXT("THREADGROUP_SIZE_X"), THREADGROUP_SIZE_X);
+	OutEnvironment.SetDefine(TEXT("MAX_NUMBER_DIMENSIONS"), MAX_NUMBER_DIMENSIONS);
+}
+
+
+
+/* Shader implementation
+ *****************************************************************************/
+
+IMPLEMENT_SHADER_TYPE(, FConvBaseCS, TEXT("/Plugins/NeuralNetworkInference/Private/ConvBaseOperator.usf"), TEXT("ConvBaseCS"), SF_Compute) // Path defined in NeuralNetworkInferenceShadersModule.cpp
diff --git a/Engine/Plugins/Experimental/NeuralNetworkInference/Source/Private/NeuralNetworkInferenceShaders/Private/ConvTransposeCS.cpp b/Engine/Plugins/Experimental/NeuralNetworkInference/Source/Private/NeuralNetworkInferenceShaders/Private/ConvTransposeCS.cpp
new file mode 100644
--- /dev/null
+++ b/Engine/Plugins/Experimental/NeuralNetworkInference/Source/Private/NeuralNetworkInferenceShaders/Private/ConvTransposeCS.cpp
@@ -0,0 +1,28 @@
+// Copyright Epic Games, Inc. All Rights Reserved.
+
+#include "ConvTransposeCS.h"
+
+
+
+/* FConvTransposeCS static members
+ *****************************************************************************/
+
+const uint32 FConvTransposeCS::THREADGROUP_SIZE_X(128);
+
+
+
+/* FConvTransposeCS public functions
+ *****************************************************************************/
+
+void FConvTransposeCS::ModifyCompilationEnvironment(const FGlobalShaderPermutationParameters& InParameters, FShaderCompilerEnvironment& OutEnvironment)
+{
+	FGlobalShader::ModifyCompilationEnvironment(InParameters, OutEnvironment);
+	OutEnvironment.SetDefine(TEXT("THREADGROUP_SIZE_X"), THREADGROUP_SIZE_X);
+}
+
+
+
+/* Shader implementation
+ *****************************************************************************/
+
+IMPLEMENT_SHADER_TYPE(, FConvTransposeCS, TEXT("/Plugins/NeuralNetworkInference/Private/ConvTransposeOperator.usf"), TEXT("XToXWithZerosCS"), SF_Compute) // Path defined in NeuralNetworkInferenceShadersModule.cpp
diff --git a/Engine/Plugins/Experimental/NeuralNetworkInference/Source/Private/NeuralNetworkInferenceShaders/Private/CopyCS.cpp b/Engine/Plugins/Experimental/NeuralNetworkInference/Source/Private/NeuralNetworkInferenceShaders/Private/CopyCS.cpp
new file mode 100644
--- /dev/null
+++ b/Engine/Plugins/Experimental/NeuralNetworkInference/Source/Private/NeuralNetworkInferenceShaders/Private/CopyCS.cpp
@@ -0,0 +1,25 @@
+// Copyright Epic Games, Inc. All Rights Reserved.
+
+#include "CopyCS.h"
+
+
+
+/* FCopyCS public functions
+ *****************************************************************************/
+
+void FCopyCS::ModifyCompilationEnvironment(const FGlobalShaderPermutationParameters& InParameters, FShaderCompilerEnvironment& OutEnvironment)
+{
+	FGlobalShader::ModifyCompilationEnvironment(InParameters, OutEnvironment);
+	OutEnvironment.SetDefine(TEXT("THREADGROUP_SIZE_X"), THREADGROUP_SIZE_X);
+}
+
+
+
+const uint32 FCopyCS::THREADGROUP_SIZE_X = 128;
+
+
+
+/* Shader implementation
+ *****************************************************************************/
+
+IMPLEMENT_SHADER_TYPE(, FCopyCS, TEXT("/Plugins/NeuralNetworkInference/Private/Copy.usf"), TEXT("CopyCS"), SF_Compute) // Path defined in NeuralNetworkInferenceShadersModule.cpp
diff --git a/Engine/Plugins/Experimental/NeuralNetworkInference/Source/Private/NeuralNetworkInferenceShaders/Private/ElementWiseCS.cpp b/Engine/Plugins/Experimental/NeuralNetworkInference/Source/Private/NeuralNetworkInferenceShaders/Private/ElementWiseCS.cpp
new file mode 100644
--- /dev/null
+++ b/Engine/Plugins/Experimental/NeuralNetworkInference/Source/Private/NeuralNetworkInferenceShaders/Private/ElementWiseCS.cpp
@@ -0,0 +1,115 @@
+// Copyright Epic Games, Inc. All Rights Reserved.
+
+#include "ElementWiseCS.h"
+#include "Utils.h"
+
+
+
+/* FElementWiseCS static members
+ *****************************************************************************/
+
+const uint32 FElementWiseCS::THREADGROUP_SIZE_X(128);
+
+
+
+/* FElementWiseCS public functions
+ *****************************************************************************/
+
+void FElementWiseCS::ModifyCompilationEnvironment(const FGlobalShaderPermutationParameters& InParameters, FShaderCompilerEnvironment& OutEnvironment)
+{
+	FGlobalShader::ModifyCompilationEnvironment(InParameters, OutEnvironment);
+	OutEnvironment.SetDefine(TEXT("THREADGROUP_SIZE_X"), THREADGROUP_SIZE_X);
+	// NumberAttributes
+	FPermutationDomain PermutationVector(InParameters.PermutationId);
+	const uint32 NumberAttributes = GetNumberAttributes(PermutationVector.Get<FShaderType>());
+	OutEnvironment.SetDefine(TEXT("NUMBER_ATTRIBUTES"), NumberAttributes);
+	// ElementWiseFunction
+	const FString ElementWiseFunction = GetFunctionString(PermutationVector.Get<FShaderType>());
+	if (NumberAttributes == 0)
+	{
+		OutEnvironment.SetDefine(TEXT("ElementWiseFunction(X)"), *ElementWiseFunction);
+	}
+	else
+	{
+		OutEnvironment.SetDefine(TEXT("ElementWiseFunction(X, Attribute)"), *ElementWiseFunction);
+	}
+	// Note: IS_INLINED is already defined because we used SHADER_PERMUTATION_BOOL("IS_INLINED")
+}
+
+
+
+/* FElementWiseCS private static functions
+ *****************************************************************************/
+
+FString FElementWiseCS::GetFunctionString(const EElementWiseOperator InType)
+{
+	if (InType >= EElementWiseOperator::MAX)
+	{
+		UE_LOG(LogNeuralNetworkInferenceShaders, Warning, TEXT("FElementWiseCS::GetFunctionString(): Unexpected InType = %d."), (int32)InType);
+		return TEXT("FElementWiseCS::GetFunctionString(): Unexpected InType.");
+	}
+	const TArray<FString> FunctionStrings({
+		TEXT("abs(X)"),						// Abs
+		TEXT("acos(X)"),					// Acos
+		TEXT("asin(X)"),					// Asin
+		TEXT("atan(X)"),					// Atan
+		TEXT("ceil(X)"),					// Ceil
+		TEXT("cos(X)"),						// Cos
+		TEXT("cosh(X)"),					// Cosh
+		TEXT("exp(X)"),						// Exp
+		TEXT("floor(X)"),					// Floor
+		TEXT("(X < 0 ? Attribute * X : X)"),// LeakyRelu
+		TEXT("log(X)"),						// Log
+		TEXT("-(X)"),						// Neg
+		TEXT("1/(X)"),						// Reciprocal
+		TEXT("round(X)"),					// Round
+		TEXT("(X < 0 ? 0 : X)"),			// Relu
+		TEXT("1/(1+exp(-(X)))"),			// Sigmoid
+		TEXT("sign(X)"),					// Sign
+		TEXT("sin(X)"),						// Sin
+		TEXT("sinh(X)"),					// Sinh
+		TEXT("sqrt(X)"),					// Sqrt
+		TEXT("tan(X)"),						// Tan
+		TEXT("tanh(X)") });					// Tanh
+	return FunctionStrings[(int32)InType];
+}
+
+uint32 FElementWiseCS::GetNumberAttributes(const EElementWiseOperator InType)
+{
+	if (InType >= EElementWiseOperator::MAX)
+	{
+		UE_LOG(LogNeuralNetworkInferenceShaders, Warning, TEXT("FElementWiseCS::GetNumberAttributes(): Unexpected InType = %d."), (int32)InType);
+		return 0;
+	}
+	const TArray<uint32> NumberAttributes({
+		0,	// Abs
+		0,	// Acos
+		0,	// Asin
+		0,	// Atan
+		0,	// Ceil
+		0,	// Cos
+		0,	// Cosh
+		0,	// Exp
+		0,	// Floor
+		1,	// LeakyRelu
+		0,	// Log
+		0,	// Neg
+		0,	// Reciprocal
+		0,	// Round
+		0,	// Relu
+		0,	// Sigmoid
+		0,	// Sign
+		0,	// Sin
+		0,	// Sinh
+		0,	// Sqrt
+		0,	// Tan
+		0 });// Tanh
+	return NumberAttributes[(int32)InType];
+}
+
+
+
+/* Shader implementation
+ *****************************************************************************/
+
+IMPLEMENT_SHADER_TYPE(, FElementWiseCS, TEXT("/Plugins/NeuralNetworkInference/Private/ElementWiseOperator.usf"), TEXT("ElementWiseCS"), SF_Compute) // Path defined in NeuralNetworkInferenceShadersModule.cpp
diff --git a/Engine/Plugins/Experimental/NeuralNetworkInference/Source/Private/NeuralNetworkInferenceShaders/Private/GemmCS.cpp b/Engine/Plugins/Experimental/NeuralNetworkInference/Source/Private/NeuralNetworkInferenceShaders/Private/GemmCS.cpp
new file mode 100644
--- /dev/null
+++ b/Engine/Plugins/Experimental/NeuralNetworkInference/Source/Private/NeuralNetworkInferenceShaders/Private/GemmCS.cpp
@@ -0,0 +1,27 @@
+// Copyright Epic Games, Inc. All Rights Reserved.
+
+#include "GemmCS.h"
+
+
+
+/* FGemmCS public functions
+ *****************************************************************************/
+
+void FGemmCS::ModifyCompilationEnvironment(const FGlobalShaderPermutationParameters& InParameters, FShaderCompilerEnvironment& OutEnvironment)
+{
+	FGlobalShader::ModifyCompilationEnvironment(InParameters, OutEnvironment);
+	OutEnvironment.SetDefine(TEXT("THREADGROUP_SIZE_X"), THREADGROUP_SIZE_X);
+	OutEnvironment.SetDefine(TEXT("THREADGROUP_SIZE_Y"), THREADGROUP_SIZE_Y);
+}
+
+
+
+const uint32 FGemmCS::THREADGROUP_SIZE_X = 128;
+const uint32 FGemmCS::THREADGROUP_SIZE_Y = 1;
+
+
+
+/* Shader implementation
+ *****************************************************************************/
+
+IMPLEMENT_SHADER_TYPE(, FGemmCS, TEXT("/Plugins/NeuralNetworkInference/Private/GemmOperator.usf"), TEXT("GemmCS"), SF_Compute) // Path defined in NeuralNetworkInferenceShadersModule.cpp
diff --git a/Engine/Plugins/Experimental/NeuralNetworkInference/Source/Private/NeuralNetworkInferenceShaders/Private/Module.cpp b/Engine/Plugins/Experimental/NeuralNetworkInference/Source/Private/NeuralNetworkInferenceShaders/Private/Module.cpp
new file mode 100644
--- /dev/null
+++ b/Engine/Plugins/Experimental/NeuralNetworkInference/Source/Private/NeuralNetworkInferenceShaders/Private/Module.cpp
@@ -0,0 +1,36 @@
+// Copyright Epic Games, Inc. All Rights Reserved.
+
+#include "Module.h"
+#include "Interfaces/IPluginManager.h"
+#include "ShaderCore.h"
+#include "Utils.h"
+
+#define LOCTEXT_NAMESPACE "NeuralNetworkInferenceShaders"
+
+// This code will execute after your module is loaded into memory; the exact timing is specified in the .uplugin file per-module
+void FNeuralNetworkInferenceShadersModule::StartupModule()
+{
+	// Add shader directory
+	const TSharedPtr<IPlugin> NeuralNetworkInferencePlugin = IPluginManager::Get().FindPlugin(TEXT("NeuralNetworkInference"));
+	if (NeuralNetworkInferencePlugin.IsValid())
+	{
+		const FString RealShaderDirectory = NeuralNetworkInferencePlugin->GetBaseDir() + TEXT("/Shaders/"); // TEXT("../../../Engine/Plugins/NeuralNetworkInference/Shaders");
+		const FString VirtualShaderDirectory = TEXT("/Plugins/NeuralNetworkInference");
+		AddShaderSourceDirectoryMapping(VirtualShaderDirectory, RealShaderDirectory);
+	}
+	else
+	{
+		UE_LOG(LogNeuralNetworkInferenceShaders, Warning,
+			TEXT("FNeuralNetworkInferenceShadersModule::StartupModule(): NeuralNetworkInferencePlugin was nullptr, shaders directory not added."));
+	}
+}
+
+// This function may be called during shutdown to clean up your module. For modules that support dynamic reloading,
+// we call this function before unloading the module.
+void FNeuralNetworkInferenceShadersModule::ShutdownModule()
+{
+}
+
+IMPLEMENT_MODULE(FNeuralNetworkInferenceShadersModule, NeuralNetworkInferenceShaders);
+
+#undef LOCTEXT_NAMESPACE
diff --git a/Engine/Plugins/Experimental/NeuralNetworkInference/Source/Private/NeuralNetworkInferenceShaders/Private/Module.h b/Engine/Plugins/Experimental/NeuralNetworkInference/Source/Private/NeuralNetworkInferenceShaders/Private/Module.h
new file mode 100644
--- /dev/null
+++ b/Engine/Plugins/Experimental/NeuralNetworkInference/Source/Private/NeuralNetworkInferenceShaders/Private/Module.h
@@ -0,0 +1,14 @@
+// Copyright Epic Games, Inc. All Rights Reserved.
+
+#pragma once
+
+#include "CoreMinimal.h"
+#include "Modules/ModuleManager.h"
+
+class FNeuralNetworkInferenceShadersModule : public IModuleInterface
+{
+public:
+	// Begin IModuleInterface
+	virtual void StartupModule() override;
+	virtual void ShutdownModule() override;
+};
diff --git a/Engine/Plugins/Experimental/NeuralNetworkInference/Source/Private/NeuralNetworkInferenceShaders/Private/MultidirectionalBroadcastCS.cpp b/Engine/Plugins/Experimental/NeuralNetworkInference/Source/Private/NeuralNetworkInferenceShaders/Private/MultidirectionalBroadcastCS.cpp
new file mode 100644
--- /dev/null
+++ b/Engine/Plugins/Experimental/NeuralNetworkInference/Source/Private/NeuralNetworkInferenceShaders/Private/MultidirectionalBroadcastCS.cpp
@@ -0,0 +1,58 @@
+// Copyright Epic Games, Inc. All Rights Reserved.
+
+#include "MultidirectionalBroadcastCS.h"
+#include "Utils.h"
+
+
+
+/* FMultidirectionalBroadcastCS static members
+ *****************************************************************************/
+
+const uint32 FMultidirectionalBroadcastCS::THREADGROUP_SIZE_X(128);
+const uint32 FMultidirectionalBroadcastCS::MAX_NUMBER_DIMENSIONS(16);
+
+
+
+/* FMultidirectionalBroadcastCS public functions
+ *****************************************************************************/
+
+void FMultidirectionalBroadcastCS::ModifyCompilationEnvironment(const FGlobalShaderPermutationParameters& InParameters, FShaderCompilerEnvironment& OutEnvironment)
+{
+	FGlobalShader::ModifyCompilationEnvironment(InParameters, OutEnvironment);
+	OutEnvironment.SetDefine(TEXT("THREADGROUP_SIZE_X"), THREADGROUP_SIZE_X);
+	OutEnvironment.SetDefine(TEXT("MAX_NUMBER_DIMENSIONS"), MAX_NUMBER_DIMENSIONS);
+
+	FPermutationDomain PermutationVector(InParameters.PermutationId);
+	const FString FunctionString = GetFunctionString(PermutationVector.Get<FShaderType>());
+	OutEnvironment.SetDefine(TEXT("MultidirectionalBroadcastFunction(X, Y)"), *FunctionString);
+
+	// Note: INLINED_MODE and SHAPE_MODE are already defined because we used SHADER_PERMUTATION_ENUM_CLASS with those names
+}
+
+
+
+/* FMultidirectionalBroadcastCS private static functions
+ *****************************************************************************/
+
+FString FMultidirectionalBroadcastCS::GetFunctionString(const EMultidirectionalBroadcastOperator InType)
+{
+	if (InType >= EMultidirectionalBroadcastOperator::MAX)
+	{
+		UE_LOG(LogNeuralNetworkInferenceShaders, Warning, TEXT("FMultidirectionalBroadcastCS::GetFunctionString(): Unexpected InType = %d."), (int32)InType);
+		return TEXT("FMultidirectionalBroadcastCS::GetFunctionString(): Unexpected InType.");
+	}
+	const TArray<FString> FunctionStrings({
+		TEXT("(X + Y)"),	// Add
+		TEXT("(X / Y)"),	// Div
+		TEXT("(X * Y)"),	// Mul
+		TEXT("pow(X, Y)"),	// Pow
+		TEXT("(X - Y)")});	// Sub
+	return FunctionStrings[(int32)InType];
+}
+
+
+
+/* Shader implementation
+ *****************************************************************************/
+
+IMPLEMENT_SHADER_TYPE(, FMultidirectionalBroadcastCS, TEXT("/Plugins/NeuralNetworkInference/Private/MultidirectionalBroadcastOperator.usf"), TEXT("MultidirectionalBroadcastCS"), SF_Compute) // Path defined in NeuralNetworkInferenceShadersModule.cpp
diff --git a/Engine/Plugins/Experimental/NeuralNetworkInference/Source/Private/NeuralNetworkInferenceShaders/Private/NeuralNetworkInferenceShaders/BatchNormalizationCS.cpp b/Engine/Plugins/Experimental/NeuralNetworkInference/Source/Private/NeuralNetworkInferenceShaders/Private/NeuralNetworkInferenceShaders/BatchNormalizationCS.cpp
deleted file mode 100644
diff --git a/Engine/Plugins/Experimental/NeuralNetworkInference/Source/Private/NeuralNetworkInferenceShaders/Private/NeuralNetworkInferenceShaders/ConvBaseCS.cpp b/Engine/Plugins/Experimental/NeuralNetworkInference/Source/Private/NeuralNetworkInferenceShaders/Private/NeuralNetworkInferenceShaders/ConvBaseCS.cpp
deleted file mode 100644
diff --git a/Engine/Plugins/Experimental/NeuralNetworkInference/Source/Private/NeuralNetworkInferenceShaders/Private/NeuralNetworkInferenceShaders/ConvTransposeCS.cpp b/Engine/Plugins/Experimental/NeuralNetworkInference/Source/Private/NeuralNetworkInferenceShaders/Private/NeuralNetworkInferenceShaders/ConvTransposeCS.cpp
deleted file mode 100644
diff --git a/Engine/Plugins/Experimental/NeuralNetworkInference/Source/Private/NeuralNetworkInferenceShaders/Private/NeuralNetworkInferenceShaders/CopyCS.cpp b/Engine/Plugins/Experimental/NeuralNetworkInference/Source/Private/NeuralNetworkInferenceShaders/Private/NeuralNetworkInferenceShaders/CopyCS.cpp
deleted file mode 100644
diff --git a/Engine/Plugins/Experimental/NeuralNetworkInference/Source/Private/NeuralNetworkInferenceShaders/Private/NeuralNetworkInferenceShaders/ElementWiseCS.cpp b/Engine/Plugins/Experimental/NeuralNetworkInference/Source/Private/NeuralNetworkInferenceShaders/Private/NeuralNetworkInferenceShaders/ElementWiseCS.cpp
deleted file mode 100644
diff --git a/Engine/Plugins/Experimental/NeuralNetworkInference/Source/Private/NeuralNetworkInferenceShaders/Private/NeuralNetworkInferenceShaders/GemmCS.cpp b/Engine/Plugins/Experimental/NeuralNetworkInference/Source/Private/NeuralNetworkInferenceShaders/Private/NeuralNetworkInferenceShaders/GemmCS.cpp
deleted file mode 100644
diff --git a/Engine/Plugins/Experimental/NeuralNetworkInference/Source/Private/NeuralNetworkInferenceShaders/Private/NeuralNetworkInferenceShaders/Module.cpp b/Engine/Plugins/Experimental/NeuralNetworkInference/Source/Private/NeuralNetworkInferenceShaders/Private/NeuralNetworkInferenceShaders/Module.cpp
deleted file mode 100644
diff --git a/Engine/Plugins/Experimental/NeuralNetworkInference/Source/Private/NeuralNetworkInferenceShaders/Private/NeuralNetworkInferenceShaders/Module.h b/Engine/Plugins/Experimental/NeuralNetworkInference/Source/Private/NeuralNetworkInferenceShaders/Private/NeuralNetworkInferenceShaders/Module.h
deleted file mode 100644
diff --git a/Engine/Plugins/Experimental/NeuralNetworkInference/Source/Private/NeuralNetworkInferenceShaders/Private/NeuralNetworkInferenceShaders/MultidirectionalBroadcastCS.cpp b/Engine/Plugins/Experimental/NeuralNetworkInference/Source/Private/NeuralNetworkInferenceShaders/Private/NeuralNetworkInferenceShaders/MultidirectionalBroadcastCS.cpp
deleted file mode 100644
diff --git a/Engine/Plugins/Experimental/NeuralNetworkInference/Source/Private/NeuralNetworkInferenceShaders/Private/NeuralNetworkInferenceShaders/Utils.h b/Engine/Plugins/Experimental/NeuralNetworkInference/Source/Private/NeuralNetworkInferenceShaders/Private/NeuralNetworkInferenceShaders/Utils.h
deleted file mode 100644
diff --git a/Engine/Plugins/Experimental/NeuralNetworkInference/Source/Private/NeuralNetworkInferenceShaders/Private/Utils.h b/Engine/Plugins/Experimental/NeuralNetworkInference/Source/Private/NeuralNetworkInferenceShaders/Private/Utils.h
new file mode 100644
--- /dev/null
+++ b/Engine/Plugins/Experimental/NeuralNetworkInference/Source/Private/NeuralNetworkInferenceShaders/Private/Utils.h
@@ -0,0 +1,7 @@
+// Copyright Epic Games, Inc. All Rights Reserved.
+
+#pragma once
+
+#include "CoreMinimal.h"
+
+DEFINE_LOG_CATEGORY_STATIC(LogNeuralNetworkInferenceShaders, Display, All);
