NeuralNetworkInference plugin:
Async task

diff --git a/Engine/Plugins/Experimental/NeuralNetworkInference/Source/NeuralNetworkInference/Private/NeuralNetwork.cpp b/Engine/Plugins/Experimental/NeuralNetworkInference/Source/NeuralNetworkInference/Private/NeuralNetwork.cpp
--- a/Engine/Plugins/Experimental/NeuralNetworkInference/Source/NeuralNetworkInference/Private/NeuralNetwork.cpp
+++ b/Engine/Plugins/Experimental/NeuralNetworkInference/Source/NeuralNetworkInference/Private/NeuralNetwork.cpp
@@ -4,6 +4,7 @@
 #include "NeuralNetworkImplBackEndUEAndORT.h"
 #include "NeuralNetworkImplBackEndUEOnly.h"
 #include "NeuralNetworkInferenceUtils.h"
+#include "NeuralNetworkAsyncTask.h"
 #include "EditorFramework/AssetImportData.h"
 #include "NeuralTimer.h"
 #include "HAL/FileManager.h"
@@ -68,6 +69,8 @@ UNeuralNetwork::~UNeuralNetwork()
 bool UNeuralNetwork::Load(const FString& InModelFilePath)
 {
 	DECLARE_SCOPE_CYCLE_COUNTER(TEXT("UNeuralNetwork_Load_FromFString"), STAT_UNeuralNetwork_Load, STATGROUP_MachineLearning);
+	
+	const FScopeLock ResourcesLock(&ResoucesCriticalSection);
 
 	// Clean previous networks
 	bIsLoaded = false;
@@ -105,6 +108,8 @@ bool UNeuralNetwork::Load(TArray<uint8>& InModelReadFromFileInBytes)
 {
 	DECLARE_SCOPE_CYCLE_COUNTER(TEXT("UNeuralNetwork_Load_FromTArrayUInt8"), STAT_UNeuralNetwork_Load, STATGROUP_MachineLearning);
 
+	const FScopeLock ResourcesLock(&ResoucesCriticalSection);
+
 	// Clean previous networks
 	bIsLoaded = false;
 
@@ -142,6 +147,7 @@ ENeuralDeviceType UNeuralNetwork::GetOutputDeviceType() const
 
 void UNeuralNetwork::SetDeviceType(const ENeuralDeviceType InDeviceType, const ENeuralDeviceType InInputDeviceType, const ENeuralDeviceType InOutputDeviceType)
 {
+	const FScopeLock ResourcesLock(&ResoucesCriticalSection);
 	if (DeviceType != InDeviceType || InputDeviceType != InInputDeviceType || OutputDeviceType != InOutputDeviceType)
 	{
 		DeviceType = InDeviceType;
@@ -161,11 +167,13 @@ ENeuralNetworkSynchronousMode UNeuralNetwork::GetSynchronousMode() const
 
 void UNeuralNetwork::SetSynchronousMode(const ENeuralNetworkSynchronousMode InSynchronousMode)
 {
+	const FScopeLock ResourcesLock(&ResoucesCriticalSection);
 	SynchronousMode = InSynchronousMode;
 }
 
 UNeuralNetwork::FOnAsyncRunCompleted& UNeuralNetwork::GetOnAsyncRunCompletedDelegate()
 {
+	const FScopeLock ResourcesLock(&ResoucesCriticalSection);
 	return OnAsyncRunCompletedDelegate;
 }
 
@@ -184,7 +192,7 @@ bool UNeuralNetwork::SetBackEnd(const ENeuralBackEnd InBackEnd)
 	BackEnd = InBackEnd;
 	const ENeuralBackEnd NewBackEndForCurrentPlatform = FPrivateNeuralNetwork::SetBackEndForCurrentPlatform(BackEnd);
 	// Reload only required if BackEndForCurrentPlatform changes (regardless of whether BackEnd changed).
-	// BackEndForCurrentPlatform does not necesarily change if BackEnd changes. E.g., changing from UEAndORT into Auto in Windows will result in BackEndForCurrentPlatform = UEAndORT in both cases.
+	// BackEndForCurrentPlatform does not necessarily change if BackEnd changes. E.g., changing from UEAndORT into Auto in Windows will result in BackEndForCurrentPlatform = UEAndORT in both cases.
 	if (BackEndForCurrentPlatform != NewBackEndForCurrentPlatform)
 	{
 		BackEndForCurrentPlatform = NewBackEndForCurrentPlatform;
@@ -252,6 +260,7 @@ const FNeuralTensor& UNeuralNetwork::GetInputTensor(const int32 InTensorIndex) c
 
 void UNeuralNetwork::SetInputFromArrayCopy(const TArray<float>& InArray, const int32 InTensorIndex)
 {
+	const FScopeLock ResourcesLock(&ResoucesCriticalSection);
 	// Sanity check
 	if (!bIsLoaded)
 	{
@@ -281,6 +290,8 @@ void UNeuralNetwork::SetInputFromArrayCopy(const TArray<float>& InArray, const i
 
 void* UNeuralNetwork::GetInputDataPointerMutable(const int32 InTensorIndex)
 {
+	const FScopeLock ResourcesLock(&ResoucesCriticalSection);
+
 	// Sanity check
 	if (!bIsLoaded)
 	{
@@ -376,6 +387,7 @@ TArray<FNeuralTensor> UNeuralNetwork::CreateInputArrayCopy() const
 
 void UNeuralNetwork::SetInputFromArrayCopy(const TArray<FNeuralTensor>& InInputTensorArray)
 {
+	const FScopeLock ResourcesLock(&ResoucesCriticalSection);
 	if (GetInputTensorNumber() != InInputTensorArray.Num())
 	{
 		UE_LOG(LogNeuralNetworkInference, Warning, TEXT("UNeuralNetwork::SetInputFromArrayCopy(): GetInputTensorNumber() == InInputTensorArray.Num() failed, %d != %d."), GetInputTensorNumber(), InInputTensorArray.Num());
@@ -404,6 +416,7 @@ TArray<FNeuralTensor> UNeuralNetwork::CreateOutputArrayCopy() const
 
 void UNeuralNetwork::InputTensorsToGPU(const TArray<int32>& InTensorIndexes)
 {
+	const FScopeLock ResourcesLock(&ResoucesCriticalSection);
 	// Sanity check
 	if (!bIsLoaded)
 	{
@@ -445,6 +458,7 @@ void UNeuralNetwork::InputTensorsToGPU(const TArray<int32>& InTensorIndexes)
 
 void UNeuralNetwork::OutputTensorsToCPU(const TArray<int32>& InTensorIndexes)
 {
+	const FScopeLock ResourcesLock(&ResoucesCriticalSection);
 	// Sanity check
 	if (!bIsLoaded)
 	{
@@ -501,7 +515,7 @@ void UNeuralNetwork::Run()
 	// UEAndORT
 	if (BackEndForCurrentPlatform == ENeuralBackEnd::UEAndORT)
 	{
-		ImplBackEndUEAndORT->Run(SynchronousMode, InputDeviceType, OutputDeviceType);
+		ImplBackEndUEAndORT->Run(OnAsyncRunCompletedDelegate, bIsBackgroundThreadRunning, ResoucesCriticalSection, SynchronousMode, DeviceType, InputDeviceType, OutputDeviceType);
 	}
 	// UEOnly
 	else if (BackEndForCurrentPlatform == ENeuralBackEnd::UEOnly)
@@ -526,6 +540,7 @@ FNeuralStatsData UNeuralNetwork::GetInferenceStats() const
 	return ComputeStatsModule.GetStats();
 };
 
+
 FNeuralStatsData UNeuralNetwork::GetInputMemoryTransferStats() const
 {
 	return InputMemoryTransferStatsModule.GetStats();
@@ -539,6 +554,8 @@ bool UNeuralNetwork::Load()
 {
 	DECLARE_SCOPE_CYCLE_COUNTER(TEXT("UNeuralNetwork_Load"), STAT_UNeuralNetwork_Load, STATGROUP_MachineLearning);
 
+	const FScopeLock ResourcesLock(&ResoucesCriticalSection);
+
 	// Clean previous networks
 	bIsLoaded = false;
 
@@ -546,7 +563,8 @@ bool UNeuralNetwork::Load()
 	if (BackEndForCurrentPlatform == ENeuralBackEnd::UEAndORT)
 	{
 		UNeuralNetwork::FImplBackEndUEAndORT::WarnAndSetDeviceToCPUIfDX12NotEnabled(DeviceType);
-		bIsLoaded = UNeuralNetwork::FImplBackEndUEAndORT::Load(ImplBackEndUEAndORT, AreInputTensorSizesVariable, ModelReadFromFileInBytes, ModelFullFilePath, GetDeviceType(), GetInputDeviceType(), GetOutputDeviceType());
+		FNeuralNetworkAsyncSyncData SyncData(OnAsyncRunCompletedDelegate, bIsBackgroundThreadRunning, ResoucesCriticalSection);
+		bIsLoaded = UNeuralNetwork::FImplBackEndUEAndORT::Load(ImplBackEndUEAndORT, SyncData, AreInputTensorSizesVariable, ModelReadFromFileInBytes, ModelFullFilePath, GetDeviceType(), GetInputDeviceType(), GetOutputDeviceType());
 	}
 	// UEOnly
 	else if (BackEndForCurrentPlatform == ENeuralBackEnd::UEOnly)
@@ -578,6 +596,9 @@ FNeuralTensor& UNeuralNetwork::GetOutputTensorMutable(const int32 InTensorIndex)
 bool UNeuralNetwork::Load(TArray<FNeuralTensor>& InTensors, const TArray<FNeuralTensor*>& InInputTensors, const TArray<FNeuralTensor*>& InOutputTensors, const TArray<TSharedPtr<class FNeuralOperator>>& InOperators)
 {
 	DECLARE_SCOPE_CYCLE_COUNTER(TEXT("UNeuralNetwork_Load_FromTensorManagerAndOperators"), STAT_UNeuralNetwork_Load, STATGROUP_MachineLearning);
+
+	const FScopeLock ResourcesLock(&ResoucesCriticalSection);
+
 	BackEnd = ENeuralBackEnd::UEOnly;
 	BackEndForCurrentPlatform = ENeuralBackEnd::UEOnly;
 	// Create and load TensorManager
diff --git a/Engine/Plugins/Experimental/NeuralNetworkInference/Source/NeuralNetworkInference/Private/NeuralNetworkAsyncTask.h b/Engine/Plugins/Experimental/NeuralNetworkInference/Source/NeuralNetworkInference/Private/NeuralNetworkAsyncTask.h
new file mode 100644
--- /dev/null
+++ b/Engine/Plugins/Experimental/NeuralNetworkInference/Source/NeuralNetworkInference/Private/NeuralNetworkAsyncTask.h
@@ -0,0 +1,104 @@
+// Copyright Epic Games, Inc. All Rights Reserved.
+
+#pragma once 
+
+#include "CoreMinimal.h"
+#include "Async/AsyncWork.h"
+#include "NeuralNetwork.h"
+
+#include "ThirdPartyWarningDisabler.h"
+NNI_THIRD_PARTY_INCLUDES_START
+#undef check
+#undef TEXT
+#ifdef WITH_UE_AND_ORT_SUPPORT
+#include "core/session/onnxruntime_cxx_api.h"
+#endif //WITH_UE_AND_ORT_SUPPORT
+NNI_THIRD_PARTY_INCLUDES_END
+
+DECLARE_DELEGATE(FOnAsyncRunCompleted);
+
+struct FNeuralNetworkAsyncSyncData
+{
+	FNeuralNetworkAsyncSyncData() {}
+
+	FNeuralNetworkAsyncSyncData(FOnAsyncRunCompleted& InOutOnAsyncRunCompletedDelegate, std::atomic<bool>& bInIsBackgroundThreadRunning, FCriticalSection& InResoucesCriticalSection) 
+		: OnAsyncRunCompletedDelegate(&InOutOnAsyncRunCompletedDelegate), bIsBackgroundThreadRunning(&bInIsBackgroundThreadRunning), ResoucesCriticalSection(&InResoucesCriticalSection)
+	{ }
+
+	FOnAsyncRunCompleted* OnAsyncRunCompletedDelegate;
+	std::atomic<bool>* bIsBackgroundThreadRunning;
+	FCriticalSection* ResoucesCriticalSection;
+};
+
+#ifdef WITH_UE_AND_ORT_SUPPORT
+
+struct FNeuralNetworkAsyncOrtVariables
+{
+	FNeuralNetworkAsyncOrtVariables() {}
+
+	FNeuralNetworkAsyncOrtVariables(Ort::Session* InOutSession, const TSharedPtr<Ort::RunOptions>& InRunOptions, const TArray<Ort::Value>& InInputOrtTensors, const TArray<const char*>& InInputTensorNames,	TArray<Ort::Value>& InOutputOrtTensors,	TArray<const char*>& InOutputTensorNames)
+		: Session(InOutSession), RunOptions(InRunOptions), InputOrtTensors(&InInputOrtTensors), InputTensorNames(&InInputTensorNames), OutputOrtTensors(&InOutputOrtTensors), OutputTensorNames(&InOutputTensorNames)
+	{ }
+
+	Ort::Session* Session;
+	TSharedPtr<Ort::RunOptions> RunOptions;
+	const TArray<Ort::Value>* InputOrtTensors; 
+	const TArray<const char*>* InputTensorNames;
+	TArray<Ort::Value>* OutputOrtTensors; 
+	TArray<const char*>* OutputTensorNames;
+};
+
+
+class FNeuralNetworkAsyncTask : public FNonAbandonableTask
+{
+	friend class FAsyncTask<FNeuralNetworkAsyncTask>;
+
+public:
+	FNeuralNetworkAsyncTask(const FNeuralNetworkAsyncSyncData& InSyncData, const FNeuralNetworkAsyncOrtVariables& InOrtData) 
+		: SyncData(InSyncData), OrtData(InOrtData)
+	{ }
+
+	void SetSynchronousMode(const ENeuralNetworkSynchronousMode InSyncMode)
+	{
+		const FScopeLock ResourcesLock(SyncData.ResoucesCriticalSection);
+		SyncMode = InSyncMode;
+	}
+protected:
+	void DoWork()
+	{
+		TUniquePtr<FScopeLock> ResourcesLock;
+		if (SyncMode == ENeuralNetworkSynchronousMode::Asynchronous)
+		{
+			ResourcesLock = MakeUnique<FScopeLock>(SyncData.ResoucesCriticalSection);
+		}
+		
+		OrtData.Session->Run(
+			*(OrtData.RunOptions.Get()),
+			OrtData.InputTensorNames->GetData(),
+			OrtData.InputOrtTensors->GetData(),
+			OrtData.InputTensorNames->Num(),
+			OrtData.OutputTensorNames->GetData(),
+			OrtData.OutputOrtTensors->GetData(),
+			OrtData.OutputTensorNames->Num()
+		);
+		
+		if (SyncMode == ENeuralNetworkSynchronousMode::Asynchronous)
+		{
+			SyncData.OnAsyncRunCompletedDelegate->ExecuteIfBound();
+			SyncData.bIsBackgroundThreadRunning->store(false);
+		}
+	}
+	// This next section of code needs to be here.  Not important as to why.
+
+	FORCEINLINE TStatId GetStatId() const
+	{
+		RETURN_QUICK_DECLARE_CYCLE_STAT(FNeuralNetworkAsyncTask, STATGROUP_ThreadPoolAsyncTasks);
+	}
+
+private:
+	ENeuralNetworkSynchronousMode SyncMode;
+	const FNeuralNetworkAsyncSyncData& SyncData;
+	const FNeuralNetworkAsyncOrtVariables& OrtData;
+};
+
+#endif //WITH_UE_AND_ORT_SUPPORT
diff --git a/Engine/Plugins/Experimental/NeuralNetworkInference/Source/NeuralNetworkInference/Private/NeuralNetworkImplBackEndUEAndORT.cpp b/Engine/Plugins/Experimental/NeuralNetworkInference/Source/NeuralNetworkInference/Private/NeuralNetworkImplBackEndUEAndORT.cpp
--- a/Engine/Plugins/Experimental/NeuralNetworkInference/Source/NeuralNetworkInference/Private/NeuralNetworkImplBackEndUEAndORT.cpp
+++ b/Engine/Plugins/Experimental/NeuralNetworkInference/Source/NeuralNetworkInference/Private/NeuralNetworkImplBackEndUEAndORT.cpp
@@ -227,7 +227,7 @@ bool UNeuralNetwork::FImplBackEndUEAndORT::IsGPUConfigCompatible()
 //	return false;
 //}
 
-bool UNeuralNetwork::FImplBackEndUEAndORT::Load(TSharedPtr<FImplBackEndUEAndORT>& InOutImplBackEndUEAndORT, TArray<bool>& OutAreInputTensorSizesVariable, const TArray<uint8>& InModelReadFromFileInBytes, const FString& InModelFullFilePath, const ENeuralDeviceType InDeviceType,
+bool UNeuralNetwork::FImplBackEndUEAndORT::Load(TSharedPtr<FImplBackEndUEAndORT>& InOutImplBackEndUEAndORT, const FNeuralNetworkAsyncSyncData& InSyncData, TArray<bool>& OutAreInputTensorSizesVariable, const TArray<uint8>& InModelReadFromFileInBytes, const FString& InModelFullFilePath, const ENeuralDeviceType InDeviceType,
 	const ENeuralDeviceType InInputDeviceType, const ENeuralDeviceType InOutputDeviceType)
 {
 #ifdef WITH_UE_AND_ORT_SUPPORT
@@ -250,7 +250,6 @@ bool UNeuralNetwork::FImplBackEndUEAndORT::Load(TSharedPtr<FImplBackEndUEAndORT>
 			// Read model from ModelReadFromFileInBytesVector
 			InOutImplBackEndUEAndORT->Session = MakeUnique<Ort::Session>(*InOutImplBackEndUEAndORT->Environment, InModelReadFromFileInBytes.GetData(), InModelReadFromFileInBytes.Num(), *InOutImplBackEndUEAndORT->SessionOptions);
 
-
 #ifdef PLATFORM_WIN64
 			// Check if resource allocator is properly initialized
 			if (InOutImplBackEndUEAndORT->DmlGPUAllocator.IsValid() && !InOutImplBackEndUEAndORT->DmlGPUAllocator->IsValid())
@@ -290,6 +289,16 @@ bool UNeuralNetwork::FImplBackEndUEAndORT::Load(TSharedPtr<FImplBackEndUEAndORT>
 			return false;
 		}
 
+		// Initializing AsyncTask 
+		InOutImplBackEndUEAndORT->SyncData = InSyncData;
+		TSharedPtr<Ort::RunOptions> RunOptions = MakeShared<Ort::RunOptions>(nullptr);
+		InOutImplBackEndUEAndORT->OrtData = FNeuralNetworkAsyncOrtVariables(InOutImplBackEndUEAndORT->Session.Get(), RunOptions, InOutImplBackEndUEAndORT->InputOrtTensors, 
+			InOutImplBackEndUEAndORT->InputTensorNames, InOutImplBackEndUEAndORT->OutputOrtTensors, InOutImplBackEndUEAndORT->OutputTensorNames);
+		
+		InOutImplBackEndUEAndORT->IsAsyncTaskDone();
+		
+		InOutImplBackEndUEAndORT->NeuralNetworkAsyncTask = MakeUnique<FAsyncTask<FNeuralNetworkAsyncTask>>(InOutImplBackEndUEAndORT->SyncData, InOutImplBackEndUEAndORT->OrtData);
+		
 		return true;
 	}
 #if WITH_EDITOR
@@ -307,6 +316,16 @@ bool UNeuralNetwork::FImplBackEndUEAndORT::Load(TSharedPtr<FImplBackEndUEAndORT>
 }
 
 #ifdef WITH_UE_AND_ORT_SUPPORT
+void UNeuralNetwork::FImplBackEndUEAndORT::IsAsyncTaskDone() const 
+{
+	if (NeuralNetworkAsyncTask && !NeuralNetworkAsyncTask->IsDone())
+	{
+		UE_LOG(LogNeuralNetworkInference, Warning,
+			TEXT("FImplBackEndUEAndORT::Run(): Previous async run had not been completed. Blocking thread until it is completed."));
+		NeuralNetworkAsyncTask->EnsureCompletion(/*bDoWorkOnThisThreadIfNotStarted*/true);
+	}
+}
+
 void UNeuralNetwork::FImplBackEndUEAndORT::ClearResources()
 {
 #ifdef PLATFORM_WIN64
@@ -325,7 +344,8 @@ void UNeuralNetwork::FImplBackEndUEAndORT::ClearResources()
 }
 #endif //WITH_UE_AND_ORT_SUPPORT
 
-void UNeuralNetwork::FImplBackEndUEAndORT::Run(const ENeuralNetworkSynchronousMode InSynchronousMode, const ENeuralDeviceType InInputDeviceType, const ENeuralDeviceType InOutputDeviceType)
+void UNeuralNetwork::FImplBackEndUEAndORT::Run(FOnAsyncRunCompleted& InOutOnAsyncRunCompletedDelegate, std::atomic<bool>& bInIsBackgroundThreadRunning, FCriticalSection& InResoucesCriticalSection, 
+	const ENeuralNetworkSynchronousMode InSynchronousMode, const ENeuralDeviceType InDeviceType, const ENeuralDeviceType InInputDeviceType, const ENeuralDeviceType InOutputDeviceType)
 {
 #ifdef WITH_UE_AND_ORT_SUPPORT
 #if WITH_EDITOR
@@ -340,17 +360,19 @@ void UNeuralNetwork::FImplBackEndUEAndORT::Run(const ENeuralNetworkSynchronousMo
 			UE_LOG(LogNeuralNetworkInference, Warning, TEXT("FImplBackEndUEAndORT::Run(): InputDeviceType must be set to CPU for now."));
 			return;
 		}
+		
+		IsAsyncTaskDone();
+		NeuralNetworkAsyncTask->GetTask().SetSynchronousMode(InSynchronousMode);
 
 		// Run UNeuralNetwork
 		if (InSynchronousMode == ENeuralNetworkSynchronousMode::Synchronous)
 		{
-			Session->Run(Ort::RunOptions{ nullptr },
-				InputTensorNames.GetData(), &InputOrtTensors[0], InputTensorNames.Num(),
-				OutputTensorNames.GetData(), &OutputOrtTensors[0], OutputTensorNames.Num());
+			NeuralNetworkAsyncTask->StartSynchronousTask();
 		}
 		else if (InSynchronousMode == ENeuralNetworkSynchronousMode::Asynchronous)
 		{
-			UE_LOG(LogNeuralNetworkInference, Warning, TEXT("FImplBackEndUEAndORT::Run(): SynchronousMode = %d not implemented yet. Use SynchronousMode = Synchronous."), (int32)InSynchronousMode);
+			bInIsBackgroundThreadRunning = true;
+			NeuralNetworkAsyncTask->StartBackgroundTask();
 		}
 		else
 		{
diff --git a/Engine/Plugins/Experimental/NeuralNetworkInference/Source/NeuralNetworkInference/Private/NeuralNetworkImplBackEndUEAndORT.h b/Engine/Plugins/Experimental/NeuralNetworkInference/Source/NeuralNetworkInference/Private/NeuralNetworkImplBackEndUEAndORT.h
--- a/Engine/Plugins/Experimental/NeuralNetworkInference/Source/NeuralNetworkInference/Private/NeuralNetworkImplBackEndUEAndORT.h
+++ b/Engine/Plugins/Experimental/NeuralNetworkInference/Source/NeuralNetworkInference/Private/NeuralNetworkImplBackEndUEAndORT.h
@@ -3,40 +3,45 @@
 #pragma once
 
 #include "NeuralNetwork.h"
+#include "Async/AsyncWork.h"
 
 // RHI includes must happen before onnxruntime_cxx_api.h (both include Windows.h)
 #include "HAL/CriticalSection.h"
 #include "RHI.h"
 #include "DynamicRHI.h"
+
 #if defined(WITH_UE_AND_ORT_SUPPORT) && defined(PLATFORM_WIN64)
 	// Disable NOMINMAX & WIN32_LEAN_AND_MEAN defines to avoid compiler warnings
-	#pragma push_macro("NOMINMAX")
-	#pragma push_macro("WIN32_LEAN_AND_MEAN")
-	#undef NOMINMAX
-	#undef WIN32_LEAN_AND_MEAN
-	#include "D3D12RHIPrivate.h"
-	#pragma pop_macro("WIN32_LEAN_AND_MEAN")
-	#pragma pop_macro("NOMINMAX")
+#pragma push_macro("NOMINMAX")
+#pragma push_macro("WIN32_LEAN_AND_MEAN")
+#undef NOMINMAX
+#undef WIN32_LEAN_AND_MEAN
+#include "D3D12RHIPrivate.h"
+#pragma pop_macro("WIN32_LEAN_AND_MEAN")
+#pragma pop_macro("NOMINMAX")
 #endif
 
+#include "NeuralNetworkAsyncTask.h"
+
 #include "ThirdPartyWarningDisabler.h"
 NNI_THIRD_PARTY_INCLUDES_START
 #undef check
 #undef TEXT
 #ifdef WITH_UE_AND_ORT_SUPPORT
-	#include "core/session/onnxruntime_cxx_api.h"
-	#ifdef PLATFORM_WIN64
-		struct OrtDMLProviderOptions;
-		namespace Ort
-		{
-			class DMLGPUResourceAllocator;
-		}
-	#endif
+#include "core/session/onnxruntime_cxx_api.h"
+#ifdef PLATFORM_WIN64
+struct OrtDMLProviderOptions;
+namespace Ort
+{
+	class DMLGPUResourceAllocator;
+}
+#endif
 #endif //WITH_UE_AND_ORT_SUPPORT
 NNI_THIRD_PARTY_INCLUDES_END
 
 
 
+
 struct UNeuralNetwork::FImplBackEndUEAndORT
 {
 public:
@@ -68,12 +73,20 @@ public:
 
 	static bool IsGPUConfigCompatible();
 
-	static bool Load(TSharedPtr<FImplBackEndUEAndORT>& InOutImplBackEndUEAndORT, TArray<bool>& OutAreInputTensorSizesVariable, const TArray<uint8>& InModelReadFromFileInBytes, const FString& InModelFullFilePath, const ENeuralDeviceType InDeviceType, const ENeuralDeviceType InInputDeviceType, const ENeuralDeviceType InOutputDeviceType);
+	static bool Load(TSharedPtr<FImplBackEndUEAndORT>& InOutImplBackEndUEAndORT, const FNeuralNetworkAsyncSyncData& InSyncData, TArray<bool>& OutAreInputTensorSizesVariable, const TArray<uint8>& InModelReadFromFileInBytes, const FString& InModelFullFilePath, const ENeuralDeviceType InDeviceType, const ENeuralDeviceType InInputDeviceType, const ENeuralDeviceType InOutputDeviceType);
 	
-	void Run(const ENeuralNetworkSynchronousMode InSynchronousMode, const ENeuralDeviceType InInputDeviceType, const ENeuralDeviceType InOutputDeviceType);
+	void Run(FOnAsyncRunCompleted& InOutOnAsyncRunCompletedDelegate, std::atomic<bool>& bInIsBackgroundThreadRunning , FCriticalSection& InResoucesCriticalSection,
+		const ENeuralNetworkSynchronousMode InSynchronousMode, const ENeuralDeviceType InDeviceType, const ENeuralDeviceType InInputDeviceType, const ENeuralDeviceType InOutputDeviceType);
 
 #ifdef WITH_UE_AND_ORT_SUPPORT
 private:
+
+	FNeuralNetworkAsyncSyncData SyncData;
+	FNeuralNetworkAsyncOrtVariables OrtData;
+	TUniquePtr<FAsyncTask<FNeuralNetworkAsyncTask>> NeuralNetworkAsyncTask;
+
+	void IsAsyncTaskDone() const;
+
 	static bool InitializedAndConfigureMembers(TSharedPtr<FImplBackEndUEAndORT>& InOutImplBackEndUEAndORT, const FString& InModelFullFilePath, const ENeuralDeviceType InDeviceType);
 
 	bool ConfigureMembers(const ENeuralDeviceType InDeviceType);
diff --git a/Engine/Plugins/Experimental/NeuralNetworkInference/Source/NeuralNetworkInference/Public/NeuralNetwork.h b/Engine/Plugins/Experimental/NeuralNetworkInference/Source/NeuralNetworkInference/Public/NeuralNetwork.h
--- a/Engine/Plugins/Experimental/NeuralNetworkInference/Source/NeuralNetworkInference/Public/NeuralNetwork.h
+++ b/Engine/Plugins/Experimental/NeuralNetworkInference/Source/NeuralNetworkInference/Public/NeuralNetwork.h
@@ -122,7 +122,8 @@ public:
 	 * GetBackEnd()/GetBackEndForCurrentPlatform():
 	 * - If BackEnd == Auto, GetBackEnd() will return Auto and GetBackEndForCurrentPlatform() will return the actual BackEnd being used for the current platform (UEAndORT or UEOnly).
 	 * - If BackEnd != Auto, GetBackEnd() and GetBackEndForCurrentPlatform() will both return the same value (UEAndORT or UEOnly).
-	 * SetBackEnd() will modify both BackEnd and BackEndForCurrentPlatform and return IsLoaded().
+	 * SetBackEnd() will modify both BackEnd and BackEndForCurrentPlatform and return IsLoaded(). 
+	 * Important! SetBackEnd is NOT THREAD SAFE, please ensure that no other Neural Network function such as Run() is running when SetBackEnd is called. 
 	 * @see ENeuralBackEnd for more details.
 	 */
 	ENeuralBackEnd GetBackEnd() const;
@@ -236,6 +237,11 @@ private:
 	 */
 	std::atomic<bool> bIsBackgroundThreadRunning;
 
+	/**
+	* Critical Section used to protect resources
+	*/
+	FCriticalSection ResoucesCriticalSection;
+
 	UPROPERTY()
 	TArray<uint8> ModelReadFromFileInBytes;
 
