Updated / cleaned up async code.
1. Removed auxiliary members in NNIBackeEndUEAndORT.
2. Removed auxiliary structs (they make the code harder to read).
3. Replaced pointers by references (so everything is guaranteed not to be nullptr).
4. Moved InOutImplBackEndUEAndORT->IsAsyncTaskDone() sooner to avoid crash if de-allocated sooner than it should be.
5.. (This is for Neb changes to work) Move #include "NeuralNetworkAsyncTask.h" into its cpp.

#rnx

diff --git a/Engine/Plugins/Experimental/NeuralNetworkInference/Source/NeuralNetworkInference/Private/NeuralNetwork.cpp b/Engine/Plugins/Experimental/NeuralNetworkInference/Source/NeuralNetworkInference/Private/NeuralNetwork.cpp
--- a/Engine/Plugins/Experimental/NeuralNetworkInference/Source/NeuralNetworkInference/Private/NeuralNetwork.cpp
+++ b/Engine/Plugins/Experimental/NeuralNetworkInference/Source/NeuralNetworkInference/Private/NeuralNetwork.cpp
@@ -1,15 +1,14 @@
 // Copyright Epic Games, Inc. All Rights Reserved.
 
 #include "NeuralNetwork.h"
-#include "NeuralNetworkImplBackEndUEAndORT.h"
-#include "NeuralNetworkImplBackEndUEOnly.h"
-#include "NeuralNetworkInferenceUtils.h"
-#include "NeuralNetworkAsyncTask.h"
 #include "EditorFramework/AssetImportData.h"
 #include "NeuralTimer.h"
 #include "HAL/FileManager.h"
 #include "Misc/FileHelper.h"
 #include "Misc/Paths.h"
+#include "NeuralNetworkImplBackEndUEAndORT.h"
+#include "NeuralNetworkImplBackEndUEOnly.h"
+#include "NeuralNetworkInferenceUtils.h"
 #include "RenderGraphBuilder.h"
 #include "RenderingThread.h"
 #include "RHI.h"
@@ -533,18 +532,17 @@ void UNeuralNetwork::Run()
 float UNeuralNetwork::GetLastInferenceTime() const
 {
 	return ComputeStatsModule.GetLastSample();
-};
+}
 
 FNeuralStatsData UNeuralNetwork::GetInferenceStats() const
 {
 	return ComputeStatsModule.GetStats();
-};
-
+}
 
 FNeuralStatsData UNeuralNetwork::GetInputMemoryTransferStats() const
 {
 	return InputMemoryTransferStatsModule.GetStats();
-};
+}
 
 
 /* UNeuralNetwork private functions
@@ -563,8 +561,8 @@ bool UNeuralNetwork::Load()
 	if (BackEndForCurrentPlatform == ENeuralBackEnd::UEAndORT)
 	{
 		UNeuralNetwork::FImplBackEndUEAndORT::WarnAndSetDeviceToCPUIfDX12NotEnabled(DeviceType);
-		FNeuralNetworkAsyncSyncData SyncData(OnAsyncRunCompletedDelegate, bIsBackgroundThreadRunning, ResoucesCriticalSection);
-		bIsLoaded = UNeuralNetwork::FImplBackEndUEAndORT::Load(ImplBackEndUEAndORT, SyncData, AreInputTensorSizesVariable, ModelReadFromFileInBytes, ModelFullFilePath, GetDeviceType(), GetInputDeviceType(), GetOutputDeviceType());
+		bIsLoaded = UNeuralNetwork::FImplBackEndUEAndORT::Load(ImplBackEndUEAndORT, OnAsyncRunCompletedDelegate, bIsBackgroundThreadRunning, ResoucesCriticalSection, AreInputTensorSizesVariable,
+			ModelReadFromFileInBytes, ModelFullFilePath, GetDeviceType(), GetInputDeviceType(), GetOutputDeviceType());
 	}
 	// UEOnly
 	else if (BackEndForCurrentPlatform == ENeuralBackEnd::UEOnly)
diff --git a/Engine/Plugins/Experimental/NeuralNetworkInference/Source/NeuralNetworkInference/Private/NeuralNetworkAsyncTask.h b/Engine/Plugins/Experimental/NeuralNetworkInference/Source/NeuralNetworkInference/Private/NeuralNetworkAsyncTask.h
--- a/Engine/Plugins/Experimental/NeuralNetworkInference/Source/NeuralNetworkInference/Private/NeuralNetworkAsyncTask.h
+++ b/Engine/Plugins/Experimental/NeuralNetworkInference/Source/NeuralNetworkInference/Private/NeuralNetworkAsyncTask.h
@@ -4,7 +4,7 @@
 
 #include "CoreMinimal.h"
 #include "Async/AsyncWork.h"
-#include "NeuralNetwork.h"
+#include "NeuralNetworkImplBackEndUEAndORT.h" // This includes RHI which must happen before onnxruntime_cxx_api.h (both files include Windows.h)
 
 #include "ThirdPartyWarningDisabler.h"
 NNI_THIRD_PARTY_INCLUDES_START
@@ -17,88 +17,68 @@ NNI_THIRD_PARTY_INCLUDES_END
 
 DECLARE_DELEGATE(FOnAsyncRunCompleted);
 
-struct FNeuralNetworkAsyncSyncData
-{
-	FNeuralNetworkAsyncSyncData() {}
-
-	FNeuralNetworkAsyncSyncData(FOnAsyncRunCompleted& InOutOnAsyncRunCompletedDelegate, std::atomic<bool>& bInIsBackgroundThreadRunning, FCriticalSection& InResoucesCriticalSection) 
-		: OnAsyncRunCompletedDelegate(&InOutOnAsyncRunCompletedDelegate), bIsBackgroundThreadRunning(&bInIsBackgroundThreadRunning), ResoucesCriticalSection(&InResoucesCriticalSection)
-	{ }
-
-	FOnAsyncRunCompleted* OnAsyncRunCompletedDelegate;
-	std::atomic<bool>* bIsBackgroundThreadRunning;
-	FCriticalSection* ResoucesCriticalSection;
-};
-
 #ifdef WITH_UE_AND_ORT_SUPPORT
 
-struct FNeuralNetworkAsyncOrtVariables
-{
-	FNeuralNetworkAsyncOrtVariables() {}
-
-	FNeuralNetworkAsyncOrtVariables(Ort::Session* InOutSession, const TSharedPtr<Ort::RunOptions>& InRunOptions, const TArray<Ort::Value>& InInputOrtTensors, const TArray<const char*>& InInputTensorNames,	TArray<Ort::Value>& InOutputOrtTensors,	TArray<const char*>& InOutputTensorNames)
-		: Session(InOutSession), RunOptions(InRunOptions), InputOrtTensors(&InInputOrtTensors), InputTensorNames(&InInputTensorNames), OutputOrtTensors(&InOutputOrtTensors), OutputTensorNames(&InOutputTensorNames)
-	{ }
-
-	Ort::Session* Session;
-	TSharedPtr<Ort::RunOptions> RunOptions;
-	const TArray<Ort::Value>* InputOrtTensors; 
-	const TArray<const char*>* InputTensorNames;
-	TArray<Ort::Value>* OutputOrtTensors; 
-	TArray<const char*>* OutputTensorNames;
-};
-
-
 class FNeuralNetworkAsyncTask : public FNonAbandonableTask
 {
 	friend class FAsyncTask<FNeuralNetworkAsyncTask>;
 
 public:
-	FNeuralNetworkAsyncTask(const FNeuralNetworkAsyncSyncData& InSyncData, const FNeuralNetworkAsyncOrtVariables& InOrtData) 
-		: SyncData(InSyncData), OrtData(InOrtData)
+	FNeuralNetworkAsyncTask(FOnAsyncRunCompleted& InOutOnAsyncRunCompletedDelegate, std::atomic<bool>& bInOutIsBackgroundThreadRunning, FCriticalSection& InOutResoucesCriticalSection,
+		Ort::Session& InOutSession, TArray<Ort::Value>& OutOutputOrtTensors, const TArray<Ort::Value>& InInputOrtTensors, const TArray<const char*>& InInputTensorNames, const TArray<const char*>& InOutputTensorNames)
+		// Async-related
+		: OnAsyncRunCompletedDelegate(InOutOnAsyncRunCompletedDelegate), bIsBackgroundThreadRunning(bInOutIsBackgroundThreadRunning), ResoucesCriticalSection(InOutResoucesCriticalSection)
+		// ORT-related
+		, Session(InOutSession), InputOrtTensors(InInputOrtTensors), InputTensorNames(InInputTensorNames), OutputOrtTensors(OutOutputOrtTensors), OutputTensorNames(InOutputTensorNames)
 	{ }
 
 	void SetSynchronousMode(const ENeuralNetworkSynchronousMode InSyncMode)
 	{
-		const FScopeLock ResourcesLock(SyncData.ResoucesCriticalSection);
+		const FScopeLock ResourcesLock(&ResoucesCriticalSection);
 		SyncMode = InSyncMode;
 	}
+
 protected:
 	void DoWork()
 	{
 		TUniquePtr<FScopeLock> ResourcesLock;
 		if (SyncMode == ENeuralNetworkSynchronousMode::Asynchronous)
 		{
-			ResourcesLock = MakeUnique<FScopeLock>(SyncData.ResoucesCriticalSection);
+			ResourcesLock = MakeUnique<FScopeLock>(&ResoucesCriticalSection);
 		}
 		
-		OrtData.Session->Run(
-			*(OrtData.RunOptions.Get()),
-			OrtData.InputTensorNames->GetData(),
-			OrtData.InputOrtTensors->GetData(),
-			OrtData.InputTensorNames->Num(),
-			OrtData.OutputTensorNames->GetData(),
-			OrtData.OutputOrtTensors->GetData(),
-			OrtData.OutputTensorNames->Num()
-		);
+		Session.Run(Ort::RunOptions{ nullptr },
+			InputTensorNames.GetData(), InputOrtTensors.GetData(), InputTensorNames.Num(),
+			OutputTensorNames.GetData(), OutputOrtTensors.GetData(), OutputTensorNames.Num());
 		
 		if (SyncMode == ENeuralNetworkSynchronousMode::Asynchronous)
 		{
-			SyncData.OnAsyncRunCompletedDelegate->ExecuteIfBound();
-			SyncData.bIsBackgroundThreadRunning->store(false);
+			OnAsyncRunCompletedDelegate.ExecuteIfBound();
+			bIsBackgroundThreadRunning = false;
 		}
 	}
-	// This next section of code needs to be here.  Not important as to why.
 
+	// This next section of code needs to be here. Not important as to why.
 	FORCEINLINE TStatId GetStatId() const
 	{
 		RETURN_QUICK_DECLARE_CYCLE_STAT(FNeuralNetworkAsyncTask, STATGROUP_ThreadPoolAsyncTasks);
 	}
 
 private:
+	/** Variables that could change on each inference run */
 	ENeuralNetworkSynchronousMode SyncMode;
-	const FNeuralNetworkAsyncSyncData& SyncData;
-	const FNeuralNetworkAsyncOrtVariables& OrtData;
+
+	/** Async variables that could only change during construction */
+	FOnAsyncRunCompleted& OnAsyncRunCompletedDelegate;
+	std::atomic<bool>& bIsBackgroundThreadRunning;
+	FCriticalSection& ResoucesCriticalSection;
+
+	/** ORT variables that could only change during construction */
+	Ort::Session& Session;
+	const TArray<Ort::Value>& InputOrtTensors;
+	const TArray<const char*>& InputTensorNames;
+	TArray<Ort::Value>& OutputOrtTensors;
+	const TArray<const char*>& OutputTensorNames;
 };
 
 #endif //WITH_UE_AND_ORT_SUPPORT
diff --git a/Engine/Plugins/Experimental/NeuralNetworkInference/Source/NeuralNetworkInference/Private/NeuralNetworkImplBackEndUEAndORT.h b/Engine/Plugins/Experimental/NeuralNetworkInference/Source/NeuralNetworkInference/Private/NeuralNetworkImplBackEndUEAndORT.h
--- a/Engine/Plugins/Experimental/NeuralNetworkInference/Source/NeuralNetworkInference/Private/NeuralNetworkImplBackEndUEAndORT.h
+++ b/Engine/Plugins/Experimental/NeuralNetworkInference/Source/NeuralNetworkInference/Private/NeuralNetworkImplBackEndUEAndORT.h
@@ -5,7 +5,7 @@
 #include "NeuralNetwork.h"
 #include "Async/AsyncWork.h"
 
-// RHI includes must happen before onnxruntime_cxx_api.h (both include Windows.h)
+// RHI includes must happen before onnxruntime_cxx_api.h (both files include Windows.h)
 #include "HAL/CriticalSection.h"
 #include "RHI.h"
 #include "DynamicRHI.h"
@@ -21,8 +21,6 @@
 #pragma pop_macro("NOMINMAX")
 #endif
 
-#include "NeuralNetworkAsyncTask.h"
-
 #include "ThirdPartyWarningDisabler.h"
 NNI_THIRD_PARTY_INCLUDES_START
 #undef check
@@ -50,7 +48,20 @@ public:
 	 */
 	TArray<FNeuralTensor> InputTensors;
 	TArray<FNeuralTensor> OutputTensors;
+
+	static void WarnAndSetDeviceToCPUIfDX12NotEnabled(ENeuralDeviceType& InOutDeviceType);
+
+	static bool IsGPUConfigCompatible();
+
+	static bool Load(TSharedPtr<FImplBackEndUEAndORT>& InOutImplBackEndUEAndORT, FOnAsyncRunCompleted& InOutOnAsyncRunCompletedDelegate, std::atomic<bool>& bInOutIsBackgroundThreadRunning,
+		FCriticalSection& InOutResoucesCriticalSection, TArray<bool>& OutAreInputTensorSizesVariable, const TArray<uint8>& InModelReadFromFileInBytes, const FString& InModelFullFilePath,
+		const ENeuralDeviceType InDeviceType, const ENeuralDeviceType InInputDeviceType, const ENeuralDeviceType InOutputDeviceType);
+	
+	void Run(FOnAsyncRunCompleted& InOutOnAsyncRunCompletedDelegate, std::atomic<bool>& bInIsBackgroundThreadRunning , FCriticalSection& InResoucesCriticalSection,
+		const ENeuralNetworkSynchronousMode InSynchronousMode, const ENeuralDeviceType InDeviceType, const ENeuralDeviceType InInputDeviceType, const ENeuralDeviceType InOutputDeviceType);
+
 #ifdef WITH_UE_AND_ORT_SUPPORT
+private:
 	/** Network-related variables */
 	TUniquePtr<Ort::Env> Environment;
 	TUniquePtr<Ort::Session> Session;
@@ -67,23 +78,8 @@ public:
 	TArray<const char*> InputTensorNames; /* Tensor names */
 	TArray<Ort::Value> OutputOrtTensors; /* Actual ONNXRuntime tensors */
 	TArray<const char*> OutputTensorNames; /* Tensor names */
-#endif //WITH_UE_AND_ORT_SUPPORT
-
-	static void WarnAndSetDeviceToCPUIfDX12NotEnabled(ENeuralDeviceType& InOutDeviceType);
-
-	static bool IsGPUConfigCompatible();
-
-	static bool Load(TSharedPtr<FImplBackEndUEAndORT>& InOutImplBackEndUEAndORT, const FNeuralNetworkAsyncSyncData& InSyncData, TArray<bool>& OutAreInputTensorSizesVariable, const TArray<uint8>& InModelReadFromFileInBytes, const FString& InModelFullFilePath, const ENeuralDeviceType InDeviceType, const ENeuralDeviceType InInputDeviceType, const ENeuralDeviceType InOutputDeviceType);
-	
-	void Run(FOnAsyncRunCompleted& InOutOnAsyncRunCompletedDelegate, std::atomic<bool>& bInIsBackgroundThreadRunning , FCriticalSection& InResoucesCriticalSection,
-		const ENeuralNetworkSynchronousMode InSynchronousMode, const ENeuralDeviceType InDeviceType, const ENeuralDeviceType InInputDeviceType, const ENeuralDeviceType InOutputDeviceType);
-
-#ifdef WITH_UE_AND_ORT_SUPPORT
-private:
 
-	FNeuralNetworkAsyncSyncData SyncData;
-	FNeuralNetworkAsyncOrtVariables OrtData;
-	TUniquePtr<FAsyncTask<FNeuralNetworkAsyncTask>> NeuralNetworkAsyncTask;
+	TUniquePtr<FAsyncTask<class FNeuralNetworkAsyncTask>> NeuralNetworkAsyncTask;
 
 	void IsAsyncTaskDone() const;
 
