Pushed forgotten file - CIS fix
#rnx

diff --git a/Engine/Plugins/Experimental/NeuralNetworkInference/Source/NeuralNetworkInference/Private/NeuralNetworkImplBackEndUEAndORT.cpp b/Engine/Plugins/Experimental/NeuralNetworkInference/Source/NeuralNetworkInference/Private/NeuralNetworkImplBackEndUEAndORT.cpp
--- a/Engine/Plugins/Experimental/NeuralNetworkInference/Source/NeuralNetworkInference/Private/NeuralNetworkImplBackEndUEAndORT.cpp
+++ b/Engine/Plugins/Experimental/NeuralNetworkInference/Source/NeuralNetworkInference/Private/NeuralNetworkImplBackEndUEAndORT.cpp
@@ -2,6 +2,7 @@
 
 #include "NeuralNetworkImplBackEndUEAndORT.h"
 
+#include "NeuralNetworkAsyncTask.h"
 #include "NeuralNetworkInferenceUtils.h"
 #include "NeuralNetworkInferenceUtilsGPU.h"
 #include "RedirectCoutAndCerrToUeLog.h"
@@ -227,15 +228,21 @@ bool UNeuralNetwork::FImplBackEndUEAndORT::IsGPUConfigCompatible()
 //	return false;
 //}
 
-bool UNeuralNetwork::FImplBackEndUEAndORT::Load(TSharedPtr<FImplBackEndUEAndORT>& InOutImplBackEndUEAndORT, const FNeuralNetworkAsyncSyncData& InSyncData, TArray<bool>& OutAreInputTensorSizesVariable, const TArray<uint8>& InModelReadFromFileInBytes, const FString& InModelFullFilePath, const ENeuralDeviceType InDeviceType,
-	const ENeuralDeviceType InInputDeviceType, const ENeuralDeviceType InOutputDeviceType)
+bool UNeuralNetwork::FImplBackEndUEAndORT::Load(TSharedPtr<FImplBackEndUEAndORT>& InOutImplBackEndUEAndORT, FOnAsyncRunCompleted& InOutOnAsyncRunCompletedDelegate, std::atomic<bool>& bInOutIsBackgroundThreadRunning,
+	FCriticalSection& InOutResoucesCriticalSection, TArray<bool>& OutAreInputTensorSizesVariable, const TArray<uint8>& InModelReadFromFileInBytes, const FString& InModelFullFilePath,
+	const ENeuralDeviceType InDeviceType, const ENeuralDeviceType InInputDeviceType, const ENeuralDeviceType InOutputDeviceType)
 {
 #ifdef WITH_UE_AND_ORT_SUPPORT
 #if WITH_EDITOR
 	try
 #endif //WITH_EDITOR
 	{
+		// Avoid multi-threaded crashes
 		const FRedirectCoutAndCerrToUeLog RedirectCoutAndCerrToUeLog;
+		if (InOutImplBackEndUEAndORT)
+		{
+			InOutImplBackEndUEAndORT->IsAsyncTaskDone();
+		}
 
 		// Initialize and configure InOutImplBackEndUEAndORT
 		if (!UNeuralNetwork::FImplBackEndUEAndORT::InitializedAndConfigureMembers(InOutImplBackEndUEAndORT, InModelFullFilePath, InDeviceType))
@@ -267,9 +274,8 @@ bool UNeuralNetwork::FImplBackEndUEAndORT::Load(TSharedPtr<FImplBackEndUEAndORT>
 		}
 		
 		// Sanity check if device type is CPU and to make sure that input and/or output is also on the CPU
-		ENeuralDeviceType	InputDeviceType = InInputDeviceType;
-		ENeuralDeviceType	OutputDeviceType = InOutputDeviceType;
-
+		ENeuralDeviceType InputDeviceType = InInputDeviceType;
+		ENeuralDeviceType OutputDeviceType = InOutputDeviceType;
 		if (InDeviceType == ENeuralDeviceType::CPU && (InInputDeviceType == ENeuralDeviceType::GPU || InOutputDeviceType == ENeuralDeviceType::GPU))
 		{
 			UE_LOG(LogNeuralNetworkInference, Warning, TEXT("FImplBackEndUEAndORT::Load(): DeviceType is CPU but Input and/or Output is set to GPU, setting all to CPU."));
@@ -290,14 +296,8 @@ bool UNeuralNetwork::FImplBackEndUEAndORT::Load(TSharedPtr<FImplBackEndUEAndORT>
 		}
 
 		// Initializing AsyncTask 
-		InOutImplBackEndUEAndORT->SyncData = InSyncData;
-		TSharedPtr<Ort::RunOptions> RunOptions = MakeShared<Ort::RunOptions>(nullptr);
-		InOutImplBackEndUEAndORT->OrtData = FNeuralNetworkAsyncOrtVariables(InOutImplBackEndUEAndORT->Session.Get(), RunOptions, InOutImplBackEndUEAndORT->InputOrtTensors, 
-			InOutImplBackEndUEAndORT->InputTensorNames, InOutImplBackEndUEAndORT->OutputOrtTensors, InOutImplBackEndUEAndORT->OutputTensorNames);
-		
-		InOutImplBackEndUEAndORT->IsAsyncTaskDone();
-		
-		InOutImplBackEndUEAndORT->NeuralNetworkAsyncTask = MakeUnique<FAsyncTask<FNeuralNetworkAsyncTask>>(InOutImplBackEndUEAndORT->SyncData, InOutImplBackEndUEAndORT->OrtData);
+		InOutImplBackEndUEAndORT->NeuralNetworkAsyncTask = MakeUnique<FAsyncTask<FNeuralNetworkAsyncTask>>(InOutOnAsyncRunCompletedDelegate, bInOutIsBackgroundThreadRunning, InOutResoucesCriticalSection,
+			*InOutImplBackEndUEAndORT->Session, InOutImplBackEndUEAndORT->OutputOrtTensors, InOutImplBackEndUEAndORT->InputOrtTensors, InOutImplBackEndUEAndORT->InputTensorNames, InOutImplBackEndUEAndORT->OutputTensorNames);
 		
 		return true;
 	}
