NN.h: better doc.

#rnx

diff --git a/Engine/Plugins/Experimental/NeuralNetworkInference/Source/NeuralNetworkInference/Public/NeuralNetwork.h b/Engine/Plugins/Experimental/NeuralNetworkInference/Source/NeuralNetworkInference/Public/NeuralNetwork.h
--- a/Engine/Plugins/Experimental/NeuralNetworkInference/Source/NeuralNetworkInference/Public/NeuralNetwork.h
+++ b/Engine/Plugins/Experimental/NeuralNetworkInference/Source/NeuralNetworkInference/Public/NeuralNetwork.h
@@ -173,7 +173,9 @@ public:
 	 * Run() executes the forward pass on the current UNeuralNetwork given the current input FDeprecatedNeuralTensor(s), which were previously filled with
 	 * SetInputFromArrayCopy() or GetInputDataPointerMutable().
 	 * Its output results can be retrieved with GetOutputTensor() or GetOutputTensors().
-	 * @param GPUSynchronousMode Whether it should block the thread until the UNeuralNetwork has fully run.
+	 *
+	 * If Run() is called asynchronously, this does not guarantee that calling SetInputFromArrayCopy multiple times will result in each one being applied for a different Run. The user is
+	 * responable of not calling SetInputFromArrayCopy until Run() is completed and its delegate (OnAsyncRunCompletedDelegate) called. Otherwise, the wrong results might be returned.
 	 */
 	void Run();
 
@@ -233,13 +235,14 @@ private:
 	/**
 	 * Whether an inference pass (i.e., Run) is happening.
 	 * This variable is thread safe as long as only "Run" modifies it. Other functions can safely read it at any time.
-	 * If other functions outside of Run() have to modify it, consider using mutexes.
+	 * If other functions outside of Run() have to modify it, consider using a mutex with a bool rather than just a std::atomic<bool>.
 	 */
 	std::atomic<bool> bIsBackgroundThreadRunning;
 
 	/**
-	* Critical Section used to protect resources
-	*/
+	 * Critical section (mutex) used to avoid issues or crashes due to the asynchronous Run being run at the same time than any other non-const class function.
+	 * @see UNeuralNetwork::Run().
+	 */
 	FCriticalSection ResoucesCriticalSection;
 
 	UPROPERTY()
