Fixed CIS issue.
#rnx

diff --git a/Engine/Plugins/Experimental/NeuralNetworkInference/Source/NeuralNetworkInference/Private/NeuralNetworkImplBackEndUEAndORT.cpp b/Engine/Plugins/Experimental/NeuralNetworkInference/Source/NeuralNetworkInference/Private/NeuralNetworkImplBackEndUEAndORT.cpp
--- a/Engine/Plugins/Experimental/NeuralNetworkInference/Source/NeuralNetworkInference/Private/NeuralNetworkImplBackEndUEAndORT.cpp
+++ b/Engine/Plugins/Experimental/NeuralNetworkInference/Source/NeuralNetworkInference/Private/NeuralNetworkImplBackEndUEAndORT.cpp
@@ -2,7 +2,6 @@
 
 #include "NeuralNetworkImplBackEndUEAndORT.h"
 
-#include "NeuralNetworkAsyncTask.h"
 #include "NeuralNetworkInferenceUtils.h"
 #include "NeuralNetworkInferenceUtilsGPU.h"
 #include "RedirectCoutAndCerrToUeLog.h"
diff --git a/Engine/Plugins/Experimental/NeuralNetworkInference/Source/NeuralNetworkInference/Private/NeuralNetworkImplBackEndUEAndORT.h b/Engine/Plugins/Experimental/NeuralNetworkInference/Source/NeuralNetworkInference/Private/NeuralNetworkImplBackEndUEAndORT.h
--- a/Engine/Plugins/Experimental/NeuralNetworkInference/Source/NeuralNetworkInference/Private/NeuralNetworkImplBackEndUEAndORT.h
+++ b/Engine/Plugins/Experimental/NeuralNetworkInference/Source/NeuralNetworkInference/Private/NeuralNetworkImplBackEndUEAndORT.h
@@ -21,6 +21,8 @@
 #pragma pop_macro("NOMINMAX")
 #endif
 
+#include "NeuralNetworkAsyncTask.h"
+
 #include "ThirdPartyWarningDisabler.h"
 NNI_THIRD_PARTY_INCLUDES_START
 #undef check
@@ -79,7 +81,7 @@ private:
 	TArray<Ort::Value> OutputOrtTensors; /* Actual ONNXRuntime tensors */
 	TArray<const char*> OutputTensorNames; /* Tensor names */
 
-	TUniquePtr<FAsyncTask<class FNeuralNetworkAsyncTask>> NeuralNetworkAsyncTask;
+	TUniquePtr<FAsyncTask<FNeuralNetworkAsyncTask>> NeuralNetworkAsyncTask;
 
 	void IsAsyncTaskDone() const;
 
