NeuralNetworkInference plugin:
Fix GPU compatible check

diff --git a/Engine/Plugins/Experimental/NeuralNetworkInference/Source/NeuralNetworkInference/Private/NeuralNetwork.cpp b/Engine/Plugins/Experimental/NeuralNetworkInference/Source/NeuralNetworkInference/Private/NeuralNetwork.cpp
--- a/Engine/Plugins/Experimental/NeuralNetworkInference/Source/NeuralNetworkInference/Private/NeuralNetwork.cpp
+++ b/Engine/Plugins/Experimental/NeuralNetworkInference/Source/NeuralNetworkInference/Private/NeuralNetwork.cpp
@@ -661,7 +661,7 @@ void UNeuralNetwork::PostLoad()
 	if (ModelReadFromFileInBytes.Num() > 0)
 	{
 		// If GPU selected but not compatible, set to CPU
-		if (!IsGPUConfigCompatibleForCurrentBackEnd())
+		if (DeviceType == ENeuralDeviceType::GPU && !IsGPUConfigCompatibleForCurrentBackEnd())
 		{
 			UE_LOG(LogNeuralNetworkInference, Warning, TEXT("UNeuralNetwork::PostLoad(): IsGPUConfigCompatibleForCurrentBackEnd() returned false, setting DeviceType to CPU."));
 			DeviceType = ENeuralDeviceType::CPU;
