Fixed CIS issue.
#rnx

diff --git a/Engine/Plugins/Experimental/NeuralNetworkInference/Source/NeuralNetworkInference/Private/NeuralNetworkAsyncTask.h b/Engine/Plugins/Experimental/NeuralNetworkInference/Source/NeuralNetworkInference/Private/NeuralNetworkAsyncTask.h
--- a/Engine/Plugins/Experimental/NeuralNetworkInference/Source/NeuralNetworkInference/Private/NeuralNetworkAsyncTask.h
+++ b/Engine/Plugins/Experimental/NeuralNetworkInference/Source/NeuralNetworkInference/Private/NeuralNetworkAsyncTask.h
@@ -4,7 +4,11 @@
 
 #include "CoreMinimal.h"
 #include "Async/AsyncWork.h"
-#include "NeuralNetworkImplBackEndUEAndORT.h" // This includes RHI which must happen before onnxruntime_cxx_api.h (both files include Windows.h)
+
+// RHI includes must happen before onnxruntime_cxx_api.h (both files include Windows.h)
+#include "HAL/CriticalSection.h"
+#include "RHI.h"
+#include "DynamicRHI.h"
 
 #include "ThirdPartyWarningDisabler.h"
 NNI_THIRD_PARTY_INCLUDES_START
