NeuralNetworkInference (NNI) plugin:
- Added support for GPU input
- Added support for PIX debugging
- Fixed incorrect pooled buffer name
- Fixed bugs in model unit tests

#rnx

diff --git a/Engine/Plugins/Experimental/NeuralNetworkInference/Source/NeuralNetworkInference/NeuralNetworkInference.Build.cs b/Engine/Plugins/Experimental/NeuralNetworkInference/Source/NeuralNetworkInference/NeuralNetworkInference.Build.cs
--- a/Engine/Plugins/Experimental/NeuralNetworkInference/Source/NeuralNetworkInference/NeuralNetworkInference.Build.cs
+++ b/Engine/Plugins/Experimental/NeuralNetworkInference/Source/NeuralNetworkInference/NeuralNetworkInference.Build.cs
@@ -80,6 +80,11 @@ public class NeuralNetworkInference : ModuleRules
 				);
 
 				AddEngineThirdPartyPrivateStaticDependencies(Target, "DX12");
+
+				if (Target.bBuildEditor == true)
+				{
+					PublicDependencyModuleNames.Add("WinPixEventRuntime");
+				}
 			}
 		}
 
diff --git a/Engine/Plugins/Experimental/NeuralNetworkInference/Source/NeuralNetworkInference/Private/NeuralNetwork.cpp b/Engine/Plugins/Experimental/NeuralNetworkInference/Source/NeuralNetworkInference/Private/NeuralNetwork.cpp
--- a/Engine/Plugins/Experimental/NeuralNetworkInference/Source/NeuralNetworkInference/Private/NeuralNetwork.cpp
+++ b/Engine/Plugins/Experimental/NeuralNetworkInference/Source/NeuralNetworkInference/Private/NeuralNetwork.cpp
@@ -514,7 +514,7 @@ void UNeuralNetwork::Run()
 	// UEAndORT
 	if (BackEndForCurrentPlatform == ENeuralBackEnd::UEAndORT)
 	{
-		ImplBackEndUEAndORT->Run(OnAsyncRunCompletedDelegate, bIsBackgroundThreadRunning, ResoucesCriticalSection, SynchronousMode, DeviceType, InputDeviceType, OutputDeviceType);
+		ImplBackEndUEAndORT->Run(SynchronousMode, DeviceType, InputDeviceType, OutputDeviceType);
 	}
 	// UEOnly
 	else if (BackEndForCurrentPlatform == ENeuralBackEnd::UEOnly)
diff --git a/Engine/Plugins/Experimental/NeuralNetworkInference/Source/NeuralNetworkInference/Private/NeuralNetworkAsyncTask.h b/Engine/Plugins/Experimental/NeuralNetworkInference/Source/NeuralNetworkInference/Private/NeuralNetworkAsyncTask.h
deleted file mode 100644
diff --git a/Engine/Plugins/Experimental/NeuralNetworkInference/Source/NeuralNetworkInference/Private/NeuralNetworkImplBackEndUEAndORT.cpp b/Engine/Plugins/Experimental/NeuralNetworkInference/Source/NeuralNetworkInference/Private/NeuralNetworkImplBackEndUEAndORT.cpp
--- a/Engine/Plugins/Experimental/NeuralNetworkInference/Source/NeuralNetworkInference/Private/NeuralNetworkImplBackEndUEAndORT.cpp
+++ b/Engine/Plugins/Experimental/NeuralNetworkInference/Source/NeuralNetworkInference/Private/NeuralNetworkImplBackEndUEAndORT.cpp
@@ -8,6 +8,26 @@
 #if WITH_EDITOR
 #include "Misc/MessageDialog.h"
 #endif //WITH_EDITOR
+#include "RedirectCoutAndCerrToUeLog.h"
+
+#if defined(WITH_UE_AND_ORT_SUPPORT) && defined(PLATFORM_WIN64)
+	#include "HAL/CriticalSection.h"
+	#include "RHI.h"
+	#include "DynamicRHI.h"
+
+	// Disable NOMINMAX & WIN32_LEAN_AND_MEAN defines to avoid compiler warnings
+	#pragma push_macro("NOMINMAX")
+	#pragma push_macro("WIN32_LEAN_AND_MEAN")
+	#pragma push_macro("UE_MINIMAL_WINDOWS_INCLUDE")
+	#undef NOMINMAX
+	#undef WIN32_LEAN_AND_MEAN
+	#define UE_MINIMAL_WINDOWS_INCLUDE // Avoids Win64 Clang warning
+	#include "D3D12RHIPrivate.h"
+	#pragma pop_macro("UE_MINIMAL_WINDOWS_INCLUDE")
+	#pragma pop_macro("WIN32_LEAN_AND_MEAN")
+	#pragma pop_macro("NOMINMAX")
+
+#endif
 
 //#define WITH_NNI_CPU_NOT_RECOMMENDED // Only for debugging purposes
 
@@ -24,9 +44,104 @@ NNI_THIRD_PARTY_INCLUDES_START
 #endif //WITH_UE_AND_ORT_SUPPORT
 NNI_THIRD_PARTY_INCLUDES_END
 
+#include "ShaderParameterUtils.h"
+#include "ShaderParameterStruct.h"
 
+#ifdef WITH_UE_AND_ORT_SUPPORT
 
-#if defined(WITH_UE_AND_ORT_SUPPORT) && defined(PLATFORM_WIN64)
+#if defined(PLATFORM_WIN64)
+
+#if WITH_EDITOR && !UE_BUILD_SHIPPING
+	#include "pix3.h"
+	#define NNIGPUProfileMarker(Name) FNNIGPUProfiler::Instance()->Marker(Name)
+#else
+	#define NNIGPUProfileMarker(Name)
+#endif
+
+#endif // PLATFORM_WIN64
+
+// Helper class to utilize the PIX CPU/GPU debugger on Windows
+class FNNIGPUProfiler
+{
+public:
+
+	static FNNIGPUProfiler* Instance()
+	{
+		static FNNIGPUProfiler	Inst;
+
+		return &Inst;
+	}
+
+	class FScopedEvent
+	{
+	public:
+
+		FScopedEvent(const FString& Name, FColor Color = FColor::Yellow)
+		{
+			FNNIGPUProfiler::Instance()->EventBegin(Name, Color);
+		}
+
+		~FScopedEvent()
+		{
+			FNNIGPUProfiler::Instance()->EventEnd();
+		}
+	};
+
+private:
+
+	FNNIGPUProfiler()
+	{
+#if defined(PLATFORM_WIN64) && defined(USE_PIX) && !defined(UE_BUILD_SHIPPING)
+		bIsEnabled = FD3D12DynamicRHI::GetD3DRHI()->IsPixEventEnabled();
+#else
+		bIsEnabled = false;
+#endif
+	}
+
+public:
+
+	~FNNIGPUProfiler()
+	{
+	}
+
+	void Marker(const FString& Name, FColor Color = FColor::Yellow)
+	{
+#if defined(PLATFORM_WIN64) && defined(USE_PIX) && !defined(UE_BUILD_SHIPPING)
+		if (bIsEnabled)
+		{
+			PIXSetMarker(PIX_COLOR(Color.R, Color.G, Color.B), Name.GetCharArray().GetData());
+		}
+#endif
+	}
+
+	void EventBegin(const FString& Name, FColor Color = FColor::Yellow)
+	{
+#if defined(PLATFORM_WIN64) && defined(USE_PIX) && !defined(UE_BUILD_SHIPPING)
+		if (bIsEnabled)
+		{
+			PIXBeginEvent(PIX_COLOR(Color.R, Color.G, Color.B), Name.GetCharArray().GetData());
+		}
+#endif
+	}
+
+	void EventEnd()
+	{
+#if defined(PLATFORM_WIN64) && defined(USE_PIX) && !defined(UE_BUILD_SHIPPING)
+		if (bIsEnabled)
+		{
+			PIXEndEvent();
+		}
+#endif
+	}
+
+private:
+
+	bool	bIsEnabled;
+
+};
+
+
+#if defined(PLATFORM_WIN64)
 
 /* FPrivateImplBackEndUEAndORT auxiliary class
  *****************************************************************************/
@@ -114,9 +229,8 @@ IDMLDevice* FPrivateImplBackEndUEAndORT::FDMLDeviceList::Add(ID3D12Device* Devic
 	return DmlDevice;
 }
 
-#endif
-
-
+#endif // PLATFORM_WIN64
+#endif // WITH_UE_AND_ORT_SUPPORT
 
 /* UNeuralNetwork public functions
  *****************************************************************************/
@@ -227,9 +341,17 @@ bool UNeuralNetwork::FImplBackEndUEAndORT::IsGPUConfigCompatible()
 //	return false;
 //}
 
-bool UNeuralNetwork::FImplBackEndUEAndORT::Load(TSharedPtr<FImplBackEndUEAndORT>& InOutImplBackEndUEAndORT, FOnAsyncRunCompleted& InOutOnAsyncRunCompletedDelegate, std::atomic<bool>& bInOutIsBackgroundThreadRunning,
-	FCriticalSection& InOutResoucesCriticalSection, TArray<bool>& OutAreInputTensorSizesVariable, const TArray<uint8>& InModelReadFromFileInBytes, const FString& InModelFullFilePath,
-	const ENeuralDeviceType InDeviceType, const ENeuralDeviceType InInputDeviceType, const ENeuralDeviceType InOutputDeviceType)
+bool UNeuralNetwork::FImplBackEndUEAndORT::Load(
+			TSharedPtr<FImplBackEndUEAndORT>& InOutImplBackEndUEAndORT, 
+			FOnAsyncRunCompleted& InOutOnAsyncRunCompletedDelegate, 
+			std::atomic<bool>& bInOutIsBackgroundThreadRunning,
+			FCriticalSection& InOutResoucesCriticalSection, 
+			TArray<bool>& OutAreInputTensorSizesVariable, 
+			const TArray<uint8>& InModelReadFromFileInBytes, 
+			const FString& InModelFullFilePath,
+			const ENeuralDeviceType InDeviceType, 
+			const ENeuralDeviceType InInputDeviceType, 
+			const ENeuralDeviceType InOutputDeviceType)
 {
 #ifdef WITH_UE_AND_ORT_SUPPORT
 #if WITH_EDITOR
@@ -238,13 +360,20 @@ bool UNeuralNetwork::FImplBackEndUEAndORT::Load(TSharedPtr<FImplBackEndUEAndORT>
 	{
 		// Avoid multi-threaded crashes
 		const FRedirectCoutAndCerrToUeLog RedirectCoutAndCerrToUeLog;
-		if (InOutImplBackEndUEAndORT)
+		
+		if (InOutImplBackEndUEAndORT.IsValid())
 		{
 			InOutImplBackEndUEAndORT->IsAsyncTaskDone();
 		}
 
 		// Initialize and configure InOutImplBackEndUEAndORT
-		if (!UNeuralNetwork::FImplBackEndUEAndORT::InitializedAndConfigureMembers(InOutImplBackEndUEAndORT, InModelFullFilePath, InDeviceType))
+		if (!UNeuralNetwork::FImplBackEndUEAndORT::InitializedAndConfigureMembers(
+					InOutImplBackEndUEAndORT, 
+					InOutOnAsyncRunCompletedDelegate, 
+					bInOutIsBackgroundThreadRunning, 
+					InOutResoucesCriticalSection, 
+					InModelFullFilePath, 
+					InDeviceType))
 		{
 			UE_LOG(LogNeuralNetworkInference, Warning, TEXT("FImplBackEndUEAndORT::Load(): InitializedAndConfigureMembers failed."));
 			return false;
@@ -282,21 +411,20 @@ bool UNeuralNetwork::FImplBackEndUEAndORT::Load(TSharedPtr<FImplBackEndUEAndORT>
 			OutputDeviceType = ENeuralDeviceType::CPU;
 		}
 
-		if (!InOutImplBackEndUEAndORT->ConfigureTensors(InOutImplBackEndUEAndORT->InputTensors, &OutAreInputTensorSizesVariable, InputDeviceType, OutputDeviceType))
+		if (!InOutImplBackEndUEAndORT->ConfigureTensors(InOutImplBackEndUEAndORT->InputTensors, &OutAreInputTensorSizesVariable, InDeviceType, InputDeviceType, OutputDeviceType))
 		{
 			UE_LOG(LogNeuralNetworkInference, Warning, TEXT("FImplBackEndUEAndORT::Load(): Failed to configure input tensors."));
 			return false;
 		}
 
-		if (!InOutImplBackEndUEAndORT->ConfigureTensors(InOutImplBackEndUEAndORT->OutputTensors, nullptr, InputDeviceType, OutputDeviceType))
+		if (!InOutImplBackEndUEAndORT->ConfigureTensors(InOutImplBackEndUEAndORT->OutputTensors, nullptr, InDeviceType, InputDeviceType, OutputDeviceType))
 		{
 			UE_LOG(LogNeuralNetworkInference, Warning, TEXT("FImplBackEndUEAndORT::Load(): Failed to configure output tensors."));
 			return false;
 		}
 
 		// Initializing AsyncTask 
-		InOutImplBackEndUEAndORT->NeuralNetworkAsyncTask = MakeUnique<FAsyncTask<FNeuralNetworkAsyncTask>>(InOutOnAsyncRunCompletedDelegate, bInOutIsBackgroundThreadRunning, InOutResoucesCriticalSection,
-			*InOutImplBackEndUEAndORT->Session, InOutImplBackEndUEAndORT->OutputOrtTensors, InOutImplBackEndUEAndORT->InputOrtTensors, InOutImplBackEndUEAndORT->InputTensorNames, InOutImplBackEndUEAndORT->OutputTensorNames);
+		InOutImplBackEndUEAndORT->NeuralNetworkAsyncTask = MakeUnique<FAsyncTask<FNeuralNetworkAsyncTask>>(InOutImplBackEndUEAndORT.Get());
 		
 		return true;
 	}
@@ -341,10 +469,17 @@ void UNeuralNetwork::FImplBackEndUEAndORT::ClearResources()
 	}
 #endif //PLATFORM_WIN64
 }
+
+UNeuralNetwork::FImplBackEndUEAndORT::FImplBackEndUEAndORT(FOnAsyncRunCompleted& InOutOnAsyncRunCompletedDelegate, std::atomic<bool>& bInIsBackgroundThreadRunning, FCriticalSection& InResoucesCriticalSection)
+	: OnAsyncRunCompletedDelegate(InOutOnAsyncRunCompletedDelegate)
+	, bIsBackgroundThreadRunning(bInIsBackgroundThreadRunning)
+	, ResoucesCriticalSection(InResoucesCriticalSection)
+{
+}
+
 #endif //WITH_UE_AND_ORT_SUPPORT
 
-void UNeuralNetwork::FImplBackEndUEAndORT::Run(FOnAsyncRunCompleted& InOutOnAsyncRunCompletedDelegate, std::atomic<bool>& bInIsBackgroundThreadRunning, FCriticalSection& InResoucesCriticalSection, 
-	const ENeuralNetworkSynchronousMode InSynchronousMode, const ENeuralDeviceType InDeviceType, const ENeuralDeviceType InInputDeviceType, const ENeuralDeviceType InOutputDeviceType)
+void UNeuralNetwork::FImplBackEndUEAndORT::Run(const ENeuralNetworkSynchronousMode InSynchronousMode, const ENeuralDeviceType InDeviceType, const ENeuralDeviceType InInputDeviceType, const ENeuralDeviceType InOutputDeviceType)
 {
 #ifdef WITH_UE_AND_ORT_SUPPORT
 #if WITH_EDITOR
@@ -353,15 +488,8 @@ void UNeuralNetwork::FImplBackEndUEAndORT::Run(FOnAsyncRunCompleted& InOutOnAsyn
 	{
 		const FRedirectCoutAndCerrToUeLog RedirectCoutAndCerrToUeLog;
 
-		// @todo: Temporarily disabled until we connect GPU input/output between UE and ORT
-		if (InInputDeviceType == ENeuralDeviceType::GPU)
-		{
-			UE_LOG(LogNeuralNetworkInference, Warning, TEXT("FImplBackEndUEAndORT::Run(): InputDeviceType must be set to CPU for now."));
-			return;
-		}
-		
 		IsAsyncTaskDone();
-		NeuralNetworkAsyncTask->GetTask().SetSynchronousMode(InSynchronousMode);
+		NeuralNetworkAsyncTask->GetTask().SetRunSessionArgs(InSynchronousMode, InDeviceType, InInputDeviceType, InOutputDeviceType);
 
 		// Run UNeuralNetwork
 		if (InSynchronousMode == ENeuralNetworkSynchronousMode::Synchronous)
@@ -370,7 +498,7 @@ void UNeuralNetwork::FImplBackEndUEAndORT::Run(FOnAsyncRunCompleted& InOutOnAsyn
 		}
 		else if (InSynchronousMode == ENeuralNetworkSynchronousMode::Asynchronous)
 		{
-			bInIsBackgroundThreadRunning = true;
+			bIsBackgroundThreadRunning = true;
 			NeuralNetworkAsyncTask->StartBackgroundTask();
 		}
 		else
@@ -390,19 +518,121 @@ void UNeuralNetwork::FImplBackEndUEAndORT::Run(FOnAsyncRunCompleted& InOutOnAsyn
 #endif //WITH_UE_AND_ORT_SUPPORT
 }
 
-
-
 /* UNeuralNetwork private functions
  *****************************************************************************/
 
 #ifdef WITH_UE_AND_ORT_SUPPORT
 
-bool UNeuralNetwork::FImplBackEndUEAndORT::InitializedAndConfigureMembers(TSharedPtr<FImplBackEndUEAndORT>& InOutImplBackEndUEAndORT, const FString& InModelFullFilePath, const ENeuralDeviceType InDeviceType)
+void UNeuralNetwork::FImplBackEndUEAndORT::RunSessionAsync(const ENeuralDeviceType InDeviceType, const ENeuralDeviceType InInputDeviceType, const ENeuralDeviceType InOutputDeviceType)
+{
+	const FScopeLock ResourcesLock(&ResoucesCriticalSection);
+
+	RunSessionImpl(InDeviceType, InInputDeviceType, InOutputDeviceType);
+
+	OnAsyncRunCompletedDelegate.ExecuteIfBound();
+	bIsBackgroundThreadRunning = false;
+}
+
+void UNeuralNetwork::FImplBackEndUEAndORT::RunSessionSync(const ENeuralDeviceType InDeviceType, const ENeuralDeviceType InInputDeviceType, const ENeuralDeviceType InOutputDeviceType)
+{
+	RunSessionImpl(InDeviceType, InInputDeviceType, InOutputDeviceType);
+}
+
+// Used when uploading tensors to GPU
+// NOTE: Upload parameter is not yet used, we plan to use it in the future
+BEGIN_SHADER_PARAMETER_STRUCT(FUploadTensorParameters, )
+	RDG_BUFFER_ACCESS(Upload, ERHIAccess::CopySrc)
+	RDG_BUFFER_ACCESS(Input, ERHIAccess::CopyDest)
+END_SHADER_PARAMETER_STRUCT()
+
+void UNeuralNetwork::FImplBackEndUEAndORT::RunSessionImpl(const ENeuralDeviceType InDeviceType, const ENeuralDeviceType InInputDeviceType, const ENeuralDeviceType InOutputDeviceType)
+{
+	if (InDeviceType == ENeuralDeviceType::GPU)
+	{
+		// Copy data to GPU (if required)
+		bool bNeedsGPUCopy = false;
+
+		for (auto& InputTensor : InputTensors)
+		{
+			if (InputTensor.GetTensorTypeGPU() != ENeuralTensorTypeGPU::Input)
+				continue;
+
+			bNeedsGPUCopy = true;
+
+			ENQUEUE_RENDER_COMMAND(UploadTensorToGPU)([this, InputTensor] (FRHICommandListImmediate& RHICmdList)
+			{
+				FRDGBuilder GraphBuilder(RHICmdList, RDG_EVENT_NAME("UploadTensorToGPU"));
+
+				// Set parameters
+				TRefCountPtr< FRDGPooledBuffer >&	PooledBuffer = InputTensor.GetPooledBuffer();
+				FRDGBufferRef						InputBufferRef = GraphBuilder.RegisterExternalBuffer(PooledBuffer);
+
+				FUploadTensorParameters*			UploadParameters = GraphBuilder.AllocParameters<FUploadTensorParameters>();
+
+				UploadParameters->Input = InputBufferRef;
+
+				GraphBuilder.AddPass(
+					RDG_EVENT_NAME("NNI:UploadTensor:%s", InputTensor.GetNameData()),
+					FUploadTensorParameters::FTypeInfo::GetStructMetadata(),
+					UploadParameters,
+					ERDGPassFlags::Copy | ERDGPassFlags::NeverCull,
+					[this, UploadParameters](FRHICommandListImmediate& RHICmdList)
+					{
+						FRHIBuffer* InputBuffer = UploadParameters->Input->GetRHI();
+
+						// NOTE: We're using UAVMask to trigger the UAV barrier in RDG
+						RHICmdList.Transition(FRHITransitionInfo(InputBuffer, ERHIAccess::CopyDest, ERHIAccess::UAVMask));
+					}
+				);
+
+				GraphBuilder.Execute();
+			});
+		}
+
+		if (bNeedsGPUCopy)
+		{
+			ENQUEUE_RENDER_COMMAND(FlushUploadTensorToGPU)([this] (FRHICommandListImmediate& RHICmdList)
+			{
+				FRDGBuilder GraphBuilder(RHICmdList, RDG_EVENT_NAME("NNI:FlushUploadTensorsToGPU"));
+
+				RHICmdList.SubmitCommandsHint();
+				RHICmdList.ImmediateFlush(EImmediateFlushType::FlushRHIThread);
+
+				GraphBuilder.Execute();
+			});
+
+			// TODO: Remove this sync point and move session run to render thread
+			FNeuralNetworkInferenceUtils::WaitUntilRHIFinished();
+		}
+	}
+
+	if (InDeviceType == ENeuralDeviceType::GPU)
+	{
+		FNNIGPUProfiler::Instance()->EventBegin("NNI:SessionRun");
+	}
+
+	Session->Run(Ort::RunOptions{ nullptr },
+		InputTensorNames.GetData(), &InputOrtTensors[0], InputTensorNames.Num(),
+		OutputTensorNames.GetData(), &OutputOrtTensors[0], OutputTensorNames.Num());
+
+	if (InDeviceType == ENeuralDeviceType::GPU)
+	{
+		FNNIGPUProfiler::Instance()->EventEnd();
+	}
+}
+
+bool UNeuralNetwork::FImplBackEndUEAndORT::InitializedAndConfigureMembers(
+			TSharedPtr<FImplBackEndUEAndORT>& InOutImplBackEndUEAndORT, 
+			FOnAsyncRunCompleted& InOutOnAsyncRunCompletedDelegate,
+			std::atomic<bool>& bInOutIsBackgroundThreadRunning,
+			FCriticalSection& InOutResoucesCriticalSection, 
+			const FString& InModelFullFilePath, 
+			const ENeuralDeviceType InDeviceType)
 {
 	// Initialize InOutImplBackEndUEAndORT
 	if (!InOutImplBackEndUEAndORT.IsValid())
 	{
-		InOutImplBackEndUEAndORT = MakeShared<FImplBackEndUEAndORT>();
+		InOutImplBackEndUEAndORT = MakeShared<FImplBackEndUEAndORT>(InOutOnAsyncRunCompletedDelegate, bInOutIsBackgroundThreadRunning, InOutResoucesCriticalSection);
 
 		// Set up ORT and create an environment
 		const char* const ModelFullFilePathCharPtr = TCHAR_TO_ANSI(*InModelFullFilePath);
@@ -521,7 +751,7 @@ bool UNeuralNetwork::FImplBackEndUEAndORT::ConfigureMembers(const ENeuralDeviceT
 	return true;
 }
 
-bool UNeuralNetwork::FImplBackEndUEAndORT::ConfigureTensors(TArray<FNeuralTensor>& OutTensors, TArray<bool>* OutAreInputTensorSizesVariable, const ENeuralDeviceType InInputDeviceType, const ENeuralDeviceType InOutputDeviceType)
+bool UNeuralNetwork::FImplBackEndUEAndORT::ConfigureTensors(TArray<FNeuralTensor>& OutTensors, TArray<bool>* OutAreInputTensorSizesVariable, const ENeuralDeviceType InDeviceType, const ENeuralDeviceType InInputDeviceType, const ENeuralDeviceType InOutputDeviceType)
 {
 	const bool bIsInput = (OutAreInputTensorSizesVariable != nullptr);
 	TArray<const char*> TensorNames;
@@ -605,19 +835,26 @@ bool UNeuralNetwork::FImplBackEndUEAndORT::ConfigureTensors(TArray<FNeuralTensor
 		TensorSizes.Push(CurrentTensorSizes);
 
 		// @todo: Should caller specify tensor GPU type?
-		// Input tensor GPU type is set to Generic
-		// Output tensor GPU type is set to Output (i.e. data should not be copied from CPU)
+		// Input/Output tensor GPU type means that data should not be copied from CPU
+		// Generic means that data is on the CPU and it will be copied to GPU
 		ENeuralTensorTypeGPU	TensorGPUType;
 
-		if (bIsInput)
+		if (InDeviceType == ENeuralDeviceType::GPU)
 		{
-			TensorGPUType = ENeuralTensorTypeGPU::Generic;
+			if (bIsInput)
+			{
+				TensorGPUType = InInputDeviceType == ENeuralDeviceType::GPU ? ENeuralTensorTypeGPU::Input : ENeuralTensorTypeGPU::Generic;
+			}
+			else
+			{
+				TensorGPUType = InOutputDeviceType == ENeuralDeviceType::GPU ? ENeuralTensorTypeGPU::Output : ENeuralTensorTypeGPU::Generic;
+			}
 		}
 		else
 		{
-			TensorGPUType = InOutputDeviceType == ENeuralDeviceType::GPU ? ENeuralTensorTypeGPU::Output : ENeuralTensorTypeGPU::Generic;
+			TensorGPUType = ENeuralTensorTypeGPU::Generic;
 		}
-
+		
 		TensorGPUTypes.Push(TensorGPUType);
 
 		CurrentTypeInfo.release();
@@ -669,6 +906,14 @@ bool UNeuralNetwork::FImplBackEndUEAndORT::SetTensorsFromNetwork(TArray<FNeuralT
 			OutTensors.Emplace(FNeuralTensor(ANSI_TO_TCHAR(TensorName), InTensorGPUTypes[TensorIndex]));
 		}
 	}
+	else
+	{
+		for (int32 TensorIndex = 0; TensorIndex < TensorNumber; ++TensorIndex)
+		{
+			OutTensors[TensorIndex].SetTensorTypeGPU(InTensorGPUTypes[TensorIndex]);
+		}
+	}
+
 	ensureMsgf(OutTensors.Num() == TensorNumber, TEXT("OutTensors.Num() == TensorNumber failed, %d != %d."), OutTensors.Num(), TensorNumber);
 
 	// Config each TensorIndex
@@ -688,7 +933,7 @@ bool UNeuralNetwork::FImplBackEndUEAndORT::SetTensorsFromNetwork(TArray<FNeuralT
 			// Link tensor with ORT blob
 			LinkTensorToONNXRuntime(OutTensors, OrtTensors, *AllocatorInfo, TensorIndex);
 		}
-		else if (InTensorGPUTypes[TensorIndex] == ENeuralTensorTypeGPU::Output)
+		else if (InTensorGPUTypes[TensorIndex] == ENeuralTensorTypeGPU::Input || InTensorGPUTypes[TensorIndex] == ENeuralTensorTypeGPU::Output)
 		{
 			// @todo: should we remove this? It's currently used to read memory from GPU to CPU
 			OutTensors[TensorIndex].SetNumUninitialized(InSizes[TensorIndex], InTensorDataTypes[TensorIndex]);
diff --git a/Engine/Plugins/Experimental/NeuralNetworkInference/Source/NeuralNetworkInference/Private/NeuralNetworkImplBackEndUEAndORT.h b/Engine/Plugins/Experimental/NeuralNetworkInference/Source/NeuralNetworkInference/Private/NeuralNetworkImplBackEndUEAndORT.h
--- a/Engine/Plugins/Experimental/NeuralNetworkInference/Source/NeuralNetworkInference/Private/NeuralNetworkImplBackEndUEAndORT.h
+++ b/Engine/Plugins/Experimental/NeuralNetworkInference/Source/NeuralNetworkInference/Private/NeuralNetworkImplBackEndUEAndORT.h
@@ -21,8 +21,6 @@
 #pragma pop_macro("NOMINMAX")
 #endif
 
-#include "NeuralNetworkAsyncTask.h"
-
 #include "ThirdPartyWarningDisabler.h"
 NNI_THIRD_PARTY_INCLUDES_START
 #undef check
@@ -40,10 +38,64 @@ namespace Ort
 NNI_THIRD_PARTY_INCLUDES_END
 
 
-
-
 struct UNeuralNetwork::FImplBackEndUEAndORT
 {
+	/// Helper class to run session as an async task
+	class FNeuralNetworkAsyncTask : public FNonAbandonableTask
+	{
+		friend class FAsyncTask<FNeuralNetworkAsyncTask>;
+
+	public:
+
+		FNeuralNetworkAsyncTask(UNeuralNetwork::FImplBackEndUEAndORT* InBackEnd)
+			: BackEnd(InBackEnd)
+		{ 
+		}
+
+		void SetRunSessionArgs(const ENeuralNetworkSynchronousMode InSyncMode, const ENeuralDeviceType InDeviceType, const ENeuralDeviceType InInputDeviceType, const ENeuralDeviceType InOutputDeviceType)
+		{
+#ifdef WITH_UE_AND_ORT_SUPPORT
+			const FScopeLock ResourcesLock(&BackEnd->ResoucesCriticalSection);
+
+			SyncMode = InSyncMode;
+			DeviceType = InDeviceType;
+			InputDeviceType = InInputDeviceType;
+			OutputDeviceType = InOutputDeviceType;
+#endif
+		}
+
+	protected:
+		void DoWork()
+		{
+#ifdef WITH_UE_AND_ORT_SUPPORT
+			if (SyncMode == ENeuralNetworkSynchronousMode::Asynchronous)
+			{
+				BackEnd->RunSessionAsync(DeviceType, InputDeviceType, OutputDeviceType);
+			}
+			else
+			{
+				BackEnd->RunSessionSync(DeviceType, InputDeviceType, OutputDeviceType);
+			}
+#endif
+		}
+
+		// This next section of code needs to be here. Not important as to why.
+		FORCEINLINE TStatId GetStatId() const
+		{
+			RETURN_QUICK_DECLARE_CYCLE_STAT(FNeuralNetworkAsyncTask, STATGROUP_ThreadPoolAsyncTasks);
+		}
+
+	private:
+		UNeuralNetwork::FImplBackEndUEAndORT* BackEnd;
+
+		/// Variables that could change on each inference run
+		ENeuralNetworkSynchronousMode SyncMode;
+		ENeuralDeviceType DeviceType;
+		ENeuralDeviceType InputDeviceType;
+		ENeuralDeviceType OutputDeviceType;
+	};
+
+
 public:
 	/**
 	 * InputTensors and OutputTensors represent the input and output TArray<FNeuralTensor> of the network, respectively.
@@ -59,11 +111,18 @@ public:
 		FCriticalSection& InOutResoucesCriticalSection, TArray<bool>& OutAreInputTensorSizesVariable, const TArray<uint8>& InModelReadFromFileInBytes, const FString& InModelFullFilePath,
 		const ENeuralDeviceType InDeviceType, const ENeuralDeviceType InInputDeviceType, const ENeuralDeviceType InOutputDeviceType);
 	
-	void Run(FOnAsyncRunCompleted& InOutOnAsyncRunCompletedDelegate, std::atomic<bool>& bInIsBackgroundThreadRunning , FCriticalSection& InResoucesCriticalSection,
-		const ENeuralNetworkSynchronousMode InSynchronousMode, const ENeuralDeviceType InDeviceType, const ENeuralDeviceType InInputDeviceType, const ENeuralDeviceType InOutputDeviceType);
+#ifdef WITH_UE_AND_ORT_SUPPORT
+	FImplBackEndUEAndORT(FOnAsyncRunCompleted& InOutOnAsyncRunCompletedDelegate, std::atomic<bool>& bInIsBackgroundThreadRunning, FCriticalSection& InResoucesCriticalSection);
+#endif
+
+	void Run(const ENeuralNetworkSynchronousMode InSynchronousMode, const ENeuralDeviceType InDeviceType, const ENeuralDeviceType InInputDeviceType, const ENeuralDeviceType InOutputDeviceType);
 
 #ifdef WITH_UE_AND_ORT_SUPPORT
 private:
+	/** Async support */
+	FOnAsyncRunCompleted& OnAsyncRunCompletedDelegate;
+	std::atomic<bool>& bIsBackgroundThreadRunning;
+	FCriticalSection& ResoucesCriticalSection;
 	/** Network-related variables */
 	TUniquePtr<Ort::Env> Environment;
 	TUniquePtr<Ort::Session> Session;
@@ -85,11 +144,17 @@ private:
 
 	void IsAsyncTaskDone() const;
 
-	static bool InitializedAndConfigureMembers(TSharedPtr<FImplBackEndUEAndORT>& InOutImplBackEndUEAndORT, const FString& InModelFullFilePath, const ENeuralDeviceType InDeviceType);
+	static bool InitializedAndConfigureMembers(
+					TSharedPtr<FImplBackEndUEAndORT>& InOutImplBackEndUEAndORT, 
+					FOnAsyncRunCompleted& InOutOnAsyncRunCompletedDelegate, 
+					std::atomic<bool>& bInOutIsBackgroundThreadRunning,
+					FCriticalSection& InOutResoucesCriticalSection, 
+					const FString& InModelFullFilePath, 
+					const ENeuralDeviceType InDeviceType);
 
 	bool ConfigureMembers(const ENeuralDeviceType InDeviceType);
 
-	bool ConfigureTensors(TArray<FNeuralTensor>& OutTensors, TArray<bool>* OutAreInputTensorSizesVariable, const ENeuralDeviceType InInputDeviceType, const ENeuralDeviceType InOutputDeviceType);
+	bool ConfigureTensors(TArray<FNeuralTensor>& OutTensors, TArray<bool>* OutAreInputTensorSizesVariable, const ENeuralDeviceType InDeviceType, const ENeuralDeviceType InInputDeviceType, const ENeuralDeviceType InOutputDeviceType);
 
 	bool SetTensorsFromNetwork(TArray<FNeuralTensor>& OutTensors, TArray<const char*>& InTensorNames, TArray<ENeuralDataType>& InTensorDataTypes, TArray<TArray<int64>>& InSizes, TArray<ENeuralTensorTypeGPU>& InTensorGPUTypes, const bool bIsInput);
 
@@ -99,6 +164,10 @@ private:
 	bool LinkTensorResourceToONNXRuntime(FNeuralTensor& InOutTensor, Ort::Value& InOutOrtTensor, void* D3DResource);
 #endif
 
+	void RunSessionAsync(const ENeuralDeviceType InDeviceType, const ENeuralDeviceType InInputDeviceType, const ENeuralDeviceType InOutputDeviceType);
+	void RunSessionSync(const ENeuralDeviceType InDeviceType, const ENeuralDeviceType InInputDeviceType, const ENeuralDeviceType InOutputDeviceType);
+	void RunSessionImpl(const ENeuralDeviceType InDeviceType, const ENeuralDeviceType InInputDeviceType, const ENeuralDeviceType InOutputDeviceType);
+
 	void ClearResources();
 
 #endif //WITH_UE_AND_ORT_SUPPORT
diff --git a/Engine/Plugins/Experimental/NeuralNetworkInference/Source/NeuralNetworkInference/Private/NeuralTensor.cpp b/Engine/Plugins/Experimental/NeuralNetworkInference/Source/NeuralNetworkInference/Private/NeuralTensor.cpp
--- a/Engine/Plugins/Experimental/NeuralNetworkInference/Source/NeuralNetworkInference/Private/NeuralTensor.cpp
+++ b/Engine/Plugins/Experimental/NeuralNetworkInference/Source/NeuralNetworkInference/Private/NeuralTensor.cpp
@@ -261,10 +261,11 @@ void FNeuralTensor::SetTensorTypeGPU(const ENeuralTensorTypeGPU InTensorTypeGPU)
 	// Sanity check
 	if (PooledBuffer.IsValid() || BufferSRVRef.IsValid() || BufferUAVRef.IsValid())
 	{
-		UE_LOG(LogNeuralNetworkInference, Warning, TEXT("FNeuralTensor-%s::SetTensorTypeGPU(): TensorTypeGPU cannot be modified from %d to %d because the GPU memory has been already initialized."
-			" Modify the GPU type before allocating the GPU memory (e.g., on the FNeuralTensor constructor)."), *Name, TensorTypeGPU, InTensorTypeGPU);
-		return;
+		PooledBuffer.Reset();
+		BufferSRVRef.Reset();
+		BufferUAVRef.Reset();
 	}
+
 	// Update TensorTypeGPU
 	TensorTypeGPU = InTensorTypeGPU;
 }
@@ -327,14 +328,33 @@ void FNeuralTensor::ToGPU_RenderThread(FRDGBuilder* InOutGraphBuilder, const EBu
 	{
 		UE_LOG(LogNeuralNetworkInference, Warning, TEXT("FNeuralTensor-%s::ToGPU_RenderThread(): bInShouldCopyFromCPU must be true for SRVs (because they cannot be edited). Assumed true."), *Name);
 	}
+
 	// Create BufferRef
-	FRDGBufferDesc BufferDesc;
-	BufferDesc.BytesPerElement = FNeuralDataTypeUtils::GetSize(DataType);
-	BufferDesc.NumElements = Num();
-	BufferDesc.Usage = InEBufferUsageFlags;
-	BufferDesc.UnderlyingType = FRDGBufferDesc::EUnderlyingType::VertexBuffer;
-	FRDGBufferRef BufferRef = (bInShouldCopyFromCPU ? CreateVertexBuffer(*InOutGraphBuilder, *Name, BufferDesc, ArrayCPU.GetData(), NumInBytes(), ERDGInitialDataFlags::NoCopy)
-		: InOutGraphBuilder->CreateBuffer(BufferDesc, *Name));
+	FRDGBufferRef BufferRef;
+
+	if (!PooledBuffer)
+	{
+		FRDGBufferDesc BufferDesc;
+		BufferDesc.BytesPerElement = FNeuralDataTypeUtils::GetSize(DataType);
+		BufferDesc.NumElements = Num();
+		BufferDesc.Usage = InEBufferUsageFlags;
+		BufferDesc.UnderlyingType = FRDGBufferDesc::EUnderlyingType::VertexBuffer;
+		
+		BufferRef = bInShouldCopyFromCPU 
+			? CreateVertexBuffer(*InOutGraphBuilder, *Name, BufferDesc, ArrayCPU.GetData(), NumInBytes(), ERDGInitialDataFlags::NoCopy)
+			: InOutGraphBuilder->CreateBuffer(BufferDesc, *Name);
+
+		// Recreate PooledBuffer for future runs
+		PooledBuffer = MakeShared<TRefCountPtr<FRDGPooledBuffer>>();
+		*PooledBuffer = InOutGraphBuilder->ConvertToExternalBuffer(BufferRef);
+	}
+	else
+	{
+		BufferRef = InOutGraphBuilder->RegisterExternalBuffer(*PooledBuffer);
+
+		InOutGraphBuilder->QueueBufferUpload(BufferRef, ArrayCPU.GetData(), NumInBytes(), ERDGInitialDataFlags::None);
+	}
+
 	// Recreate BufferSRVRef
 	if (EnumHasAnyFlags(InEBufferUsageFlags, BUF_ShaderResource))
 	{
@@ -353,9 +373,6 @@ void FNeuralTensor::ToGPU_RenderThread(FRDGBuilder* InOutGraphBuilder, const EBu
 	{
 		BufferUAVRef.Reset();
 	}
-	// Recreate PooledBuffer for future runs
-	PooledBuffer = MakeShared<TRefCountPtr<FRDGPooledBuffer>>();
-	*PooledBuffer = InOutGraphBuilder->ConvertToExternalBuffer(BufferRef);
 }
 
 void FNeuralTensor::UpdateSRVAndOrUAV_RenderThread(FRDGBuilder* InOutGraphBuilder)
@@ -429,15 +446,15 @@ bool FNeuralTensor::InitPooledBuffer(void** NativeResource)
 			BufferDesc.Usage = BUF_UnorderedAccess | BUF_ShaderResource;
 			BufferDesc.UnderlyingType = FRDGBufferDesc::EUnderlyingType::VertexBuffer;
 
-			FRDGBufferRef BufferRef = Builder.CreateBuffer(BufferDesc, *GetName());
-			
+			FRDGBufferRef BufferRef = Builder.CreateBuffer(BufferDesc, GetNameData());
+
 			// Recreate PooledBuffer for future runs
 			PooledBuffer = MakeShared<TRefCountPtr<FRDGPooledBuffer>>();
 			*PooledBuffer = Builder.ConvertToExternalBuffer(BufferRef);
 
-			// SRV and UAV can be used by callers (for example MLDeformer), look at the GetPooledBuffer()
-			BufferSRVRef = MakeShared<FRDGBufferSRVRef>(Builder.CreateSRV(BufferRef, FNeuralDataTypeUtils::GetPixelFormat(DataType)));
-			BufferUAVRef = MakeShared<FRDGBufferUAVRef>(Builder.CreateUAV(BufferRef, FNeuralDataTypeUtils::GetPixelFormat(DataType)));
+			// NOTE: SRV and UAV can be used by callers (for example MLDeformer), look at the GetPooledBuffer()
+			BufferSRVRef.Reset();
+			BufferUAVRef.Reset();
 
 			Builder.Execute();
 
@@ -467,7 +484,7 @@ TRefCountPtr<class FRDGPooledBuffer>& FNeuralTensor::GetPooledBuffer() const
 	// Sanity checks
 	checkf(bEnableGPU, TEXT("FNeuralTensor-%s::GetPooledBuffer(): bEnableGPU must be true."), *Name);
 	checkf(IsInRenderingThread(), TEXT("FNeuralTensor-%s::GetPooledBuffer(): IsInRenderingThread() must be true."), *Name);
-	checkf(BufferUAVRef.IsValid(), TEXT("FNeuralTensor-%s::GetPooledBuffer(): PooledBuffer cannot be nullptr."), *Name);
+	checkf(PooledBuffer.IsValid(), TEXT("FNeuralTensor-%s::GetPooledBuffer(): PooledBuffer cannot be nullptr."), *Name);
 	// Return PooledBuffer
 	return *PooledBuffer;
 }
diff --git a/Engine/Plugins/Experimental/NeuralNetworkInference/Source/NeuralNetworkInference/Public/NeuralTensor.h b/Engine/Plugins/Experimental/NeuralNetworkInference/Source/NeuralNetworkInference/Public/NeuralTensor.h
--- a/Engine/Plugins/Experimental/NeuralNetworkInference/Source/NeuralNetworkInference/Public/NeuralTensor.h
+++ b/Engine/Plugins/Experimental/NeuralNetworkInference/Source/NeuralNetworkInference/Public/NeuralTensor.h
@@ -148,6 +148,8 @@ public:
 
 	FORCEINLINE FString GetName() const;
 
+	FORCEINLINE const TCHAR* GetNameData() const;
+
 	FORCEINLINE ENeuralDataType GetDataType() const;
 
 	/**
@@ -398,6 +400,11 @@ FString FNeuralTensor::GetName() const
 	return Name;
 }
 
+const TCHAR* FNeuralTensor::GetNameData() const
+{
+	return Name.GetCharArray().GetData();
+}
+
 ENeuralDataType FNeuralTensor::GetDataType() const
 {
 	return DataType;
diff --git a/Engine/Plugins/Experimental/NeuralNetworkInference/Source/NeuralNetworkInferenceQA/Private/ModelUnitTester.cpp b/Engine/Plugins/Experimental/NeuralNetworkInference/Source/NeuralNetworkInferenceQA/Private/ModelUnitTester.cpp
--- a/Engine/Plugins/Experimental/NeuralNetworkInference/Source/NeuralNetworkInferenceQA/Private/ModelUnitTester.cpp
+++ b/Engine/Plugins/Experimental/NeuralNetworkInference/Source/NeuralNetworkInferenceQA/Private/ModelUnitTester.cpp
@@ -67,14 +67,45 @@ FNNIUnitTesterTimeData FModelUnitTester_GetTimeInformation(UNeuralNetwork* InOut
 
 bool FModelUnitTester::GlobalTest(const FString& InProjectContentDir, const FString& InModelZooRelativeDirectory)
 {
+	// NOTE: models are separated into multiple lines to enable/disable particular network for faster testing
+
 	// Model names, input values, and number of repetitions for profiling
-	const TArray<FString> ModelNames({ TEXT("MLRigDeformer"), TEXT("cloth_network"), TEXT("HS"), TEXT("RL") });
-	const TArray<float> InputArrayValues({ 1.f, 0.f, -1.f, 100.f, -100.f, 0.5f, -0.5f }); // This one can be shorter than CPU/GPUGroundTruths
+	const TArray<FString> ModelNames(
+	{ 
+		TEXT("MLRigDeformer"), 
+		TEXT("cloth_network"), 
+		TEXT("HS"), 
+		TEXT("RL") 
+	});
+
 	// Ground truths
-	const TArray<TArray<double>> CPUGroundTruths({ {3.728547, 0.008774, 4.595651, 212.193216, 742.434561, 4.250668, 4.717748}, {0.042571, 0.023693, 0.015783, 13.100505, 8.050994, 0.028807, 0.016387},
-		{138.372906, 126.753839, 127.287254, 130.316062, 127.303424, 124.800896, 126.546051}, {0.488662, 0.472437, 0.478862, 0.522685, 0.038322, 0.480848, 0.483821} });
-	const TArray<TArray<double>> GPUGroundTruths({ {3.728547, 0.008774, 4.595651, 212.193208, 742.434578, 4.250668, 4.717748}, {0.042571, 0.023693, 0.015783, 13.100504, 8.050994, 0.028807, 0.016387},
-		{138.373184, 126.754100, 127.287398, 130.316194, 127.303495, 124.801134, 126.5462530}, {0.488662, 0.472437, 0.478862, 0.522685, 0.038322, 0.480848, 0.483821} });
+	const TArray<TArray<double>> CPUGroundTruths(
+	{ 
+		{3.728547, 0.008774, 4.595651, 212.193216, 742.434561, 4.250668, 4.717748}, 
+		{0.042571, 0.023693, 0.015783, 13.100505, 8.050994, 0.028807, 0.016387},
+		{138.372906, 126.753839, 127.287254, 130.316062, 127.303424, 124.800896, 126.546051}, 
+		{0.488662, 0.472437, 0.478862, 0.522685, 0.038322, 0.480848, 0.483821} 
+	});
+	
+	const TArray<TArray<double>> GPUGroundTruths(
+	{ 
+		{3.728547, 0.008774, 4.595651, 212.193208, 742.434578, 4.250668, 4.717748}, 
+		{0.042571, 0.023693, 0.015783, 13.100504, 8.050994, 0.028807, 0.016387},
+		{138.373184, 126.754100, 127.287398, 130.316194, 127.303495, 124.801134, 126.5462530}, 
+		{0.488662, 0.472437, 0.478862, 0.522685, 0.038322, 0.480848, 0.483821} 
+	});
+
+	const TArray<float> InputArrayValues( // This one can be shorter than CPU/GPUGroundTruths
+	{ 
+		1.f, 
+		0.f, 
+		-1.f, 
+		100.f, 
+		-100.f, 
+		0.5f, 
+		-0.5f 
+	}); 
+	
 	// Speed profiling test - 0 repetitions means that test will not be run
 #ifdef WITH_UE_AND_ORT_SUPPORT
 	const TArray<int32> CPURepetitionsForUEAndORTBackEnd({ 1000, 1000,  50, 1000 });
@@ -287,11 +318,13 @@ bool FModelUnitTester::ModelAccuracyTest(UNeuralNetwork* InOutNetwork, const ENe
 			return false;
 		}
 	}
+	
 	// Save original network state
 	const ENeuralDeviceType OriginalDeviceType = InOutNetwork->GetDeviceType();
 	const ENeuralDeviceType OriginalInputDeviceType = InOutNetwork->GetInputDeviceType();
 	const ENeuralDeviceType OriginalOutputDeviceType = InOutNetwork->GetOutputDeviceType();
 	const ENeuralBackEnd OriginalBackEnd = InOutNetwork->GetBackEnd();
+	
 	// Set back end
 	if (!InOutNetwork->SetBackEnd(InBackEnd))
 	{
@@ -299,61 +332,58 @@ bool FModelUnitTester::ModelAccuracyTest(UNeuralNetwork* InOutNetwork, const ENe
 		return false;
 	}
 	const FString BackEndString = GetBackEndString(InBackEnd);
+	
 	// Run each input with CPU/GPU and compare with each other and with the ground truth
 	// Multiple for loops to make sure that running on the CPU does not affect GPU results or vice versa
 	bool bDidGlobalTestPassed = true;
 	TArray<TArray<float>> CPUOutputs, CPUGPUCPUOutputs, CPUGPUGPUOutputs, GPUGPUCPUOutputs;
+	
 	// Input CPU + Network CPU + Output CPU
 	for (int32 Index = 0; Index < InputArrays.Num(); ++Index)
 	{
-		InOutNetwork->SetInputFromArrayCopy(InputArrays[Index]);
 		// CPU
 		InOutNetwork->SetDeviceType(/*DeviceType*/ENeuralDeviceType::CPU, /*InputDeviceType*/ENeuralDeviceType::CPU, /*OutputDeviceType*/ENeuralDeviceType::CPU);
+		InOutNetwork->SetInputFromArrayCopy(InputArrays[Index]);
 		InOutNetwork->Run();
 		CPUOutputs.Emplace(InOutNetwork->GetOutputTensor().GetArrayCopy<float>());
 	}
+	
 	if (UNeuralNetwork::IsGPUConfigCompatibleForUEAndORTBackEnd())
 	{
 		// Input CPU + Network GPU + Output CPU
 		for (int32 Index = 0; Index < InputArrays.Num(); ++Index)
 		{
-			InOutNetwork->SetInputFromArrayCopy(InputArrays[Index]);
 			InOutNetwork->SetDeviceType(/*DeviceType*/ENeuralDeviceType::GPU, /*InputDeviceType*/ENeuralDeviceType::CPU, /*OutputDeviceType*/ENeuralDeviceType::CPU);
+			InOutNetwork->SetInputFromArrayCopy(InputArrays[Index]);
 			InOutNetwork->Run();
 			CPUGPUCPUOutputs.Emplace(InOutNetwork->GetOutputTensor().GetArrayCopy<float>());
 		}
+		
 		// Input CPU + Network GPU + Output GPU
 		for (int32 Index = 0; Index < InputArrays.Num(); ++Index)
 		{
-			InOutNetwork->SetInputFromArrayCopy(InputArrays[Index]);
 			InOutNetwork->SetDeviceType(/*DeviceType*/ENeuralDeviceType::GPU, /*InputDeviceType*/ENeuralDeviceType::CPU, /*OutputDeviceType*/ENeuralDeviceType::GPU);
+			InOutNetwork->SetInputFromArrayCopy(InputArrays[Index]);
 			InOutNetwork->Run();
 			InOutNetwork->OutputTensorsToCPU();
 			CPUGPUGPUOutputs.Emplace(InOutNetwork->GetOutputTensor().GetArrayCopy<float>());
 		}
+
 		// Input GPU + Network GPU + Output CPU
 		for (int32 Index = 0; Index < InputArrays.Num(); ++Index)
 		{
-// @todo: Remove if/else (and else contents) once DX12 version of NNI (UEAndORT) is ready
-if (InBackEnd == ENeuralBackEnd::UEOnly /*|| true*/)
-{
+			InOutNetwork->SetDeviceType(/*DeviceType*/ENeuralDeviceType::GPU, /*InputDeviceType*/ENeuralDeviceType::GPU, /*OutputDeviceType*/ENeuralDeviceType::CPU);
 			InOutNetwork->SetInputFromArrayCopy(InputArrays[Index]);
 			InOutNetwork->InputTensorsToGPU();
-			InOutNetwork->SetDeviceType(/*DeviceType*/ENeuralDeviceType::GPU, /*InputDeviceType*/ENeuralDeviceType::GPU, /*OutputDeviceType*/ENeuralDeviceType::CPU);
 			InOutNetwork->Run();
 			GPUGPUCPUOutputs.Emplace(InOutNetwork->GetOutputTensor().GetArrayCopy<float>());
-}
-else
-{
-		UE_LOG(LogNeuralNetworkInferenceQA, Display, TEXT("FModelUnitTester::ModelAccuracyTest(): GPU output for UEAndORT back end not working yet. Uncomment this line to test it."));
-		GPUGPUCPUOutputs.Push(CPUGPUCPUOutputs[Index]);
-}
 		}
 	}
 	else
 	{
 		UE_LOG(LogNeuralNetworkInferenceQA, Display, TEXT("-------------------- ModelAccuracyTest-UEAndORT-GPU skipped (DX12 not enabled)."));
 	}
+	
 	for (int32 Index = 0; Index < InputArrays.Num(); ++Index)
 	{
 		const TArray<float>& GPUGPUCPUOutput = (UNeuralNetwork::IsGPUConfigCompatibleForUEAndORTBackEnd() ? GPUGPUCPUOutputs[Index] : CPUOutputs[Index]);
