Code cleaner to read
#rnx

diff --git a/Engine/Plugins/Experimental/NeuralNetworkInference/Source/NeuralNetworkInference/NeuralNetworkInference.Build.cs b/Engine/Plugins/Experimental/NeuralNetworkInference/Source/NeuralNetworkInference/NeuralNetworkInference.Build.cs
--- a/Engine/Plugins/Experimental/NeuralNetworkInference/Source/NeuralNetworkInference/NeuralNetworkInference.Build.cs
+++ b/Engine/Plugins/Experimental/NeuralNetworkInference/Source/NeuralNetworkInference/NeuralNetworkInference.Build.cs
@@ -81,7 +81,7 @@ public class NeuralNetworkInference : ModuleRules
 
 				AddEngineThirdPartyPrivateStaticDependencies(Target, "DX12");
 
-				if (Target.bBuildEditor == true)
+				if (Target.bBuildEditor)
 				{
 					PublicDependencyModuleNames.Add("WinPixEventRuntime");
 				}
diff --git a/Engine/Plugins/Experimental/NeuralNetworkInference/Source/NeuralNetworkInference/Private/NeuralNetworkImplBackEndUEAndORT.cpp b/Engine/Plugins/Experimental/NeuralNetworkInference/Source/NeuralNetworkInference/Private/NeuralNetworkImplBackEndUEAndORT.cpp
--- a/Engine/Plugins/Experimental/NeuralNetworkInference/Source/NeuralNetworkInference/Private/NeuralNetworkImplBackEndUEAndORT.cpp
+++ b/Engine/Plugins/Experimental/NeuralNetworkInference/Source/NeuralNetworkInference/Private/NeuralNetworkImplBackEndUEAndORT.cpp
@@ -8,7 +8,6 @@
 #if WITH_EDITOR
 #include "Misc/MessageDialog.h"
 #endif //WITH_EDITOR
-#include "RedirectCoutAndCerrToUeLog.h"
 
 #if defined(WITH_UE_AND_ORT_SUPPORT) && defined(PLATFORM_WIN64)
 	#include "HAL/CriticalSection.h"
@@ -26,7 +25,6 @@
 	#pragma pop_macro("UE_MINIMAL_WINDOWS_INCLUDE")
 	#pragma pop_macro("WIN32_LEAN_AND_MEAN")
 	#pragma pop_macro("NOMINMAX")
-
 #endif
 
 //#define WITH_NNI_CPU_NOT_RECOMMENDED // Only for debugging purposes
@@ -64,11 +62,9 @@ NNI_THIRD_PARTY_INCLUDES_END
 class FNNIGPUProfiler
 {
 public:
-
 	static FNNIGPUProfiler* Instance()
 	{
-		static FNNIGPUProfiler	Inst;
-
+		static FNNIGPUProfiler Inst;
 		return &Inst;
 	}
 
@@ -88,7 +84,6 @@ public:
 	};
 
 private:
-
 	FNNIGPUProfiler()
 	{
 #if defined(PLATFORM_WIN64) && defined(USE_PIX) && !defined(UE_BUILD_SHIPPING)
@@ -99,7 +94,6 @@ private:
 	}
 
 public:
-
 	~FNNIGPUProfiler()
 	{
 	}
@@ -135,13 +129,11 @@ public:
 	}
 
 private:
-
-	bool	bIsEnabled;
-
+	bool bIsEnabled;
 };
 
 
-#if defined(PLATFORM_WIN64)
+#ifdef PLATFORM_WIN64
 
 /* FPrivateImplBackEndUEAndORT auxiliary class
  *****************************************************************************/
@@ -168,13 +160,13 @@ private:
 			IDMLDevice* DmlDevice;
 		};
 
-		TArray<DMLDeviceEntry>		Entries;
+		TArray<DMLDeviceEntry> Entries;
 	};
 };
 
 IDMLDevice* FPrivateImplBackEndUEAndORT::GetDMLDeviceThreadSafe(ID3D12Device* Device)
 {
-	static FCriticalSection CriticalSection; /* Critical section to protect GetDMLDeviceThreadSafe from being called simultaneously from multiple threads. */
+	static FCriticalSection CriticalSection; /* Protects GetDMLDeviceThreadSafe from being called simultaneously from multiple threads. */
 	static FDMLDeviceList DMLDeviceList;
 	FScopeLock Lock(&CriticalSection);
 	return DMLDeviceList.GetDMLDevice(Device);
@@ -232,6 +224,8 @@ IDMLDevice* FPrivateImplBackEndUEAndORT::FDMLDeviceList::Add(ID3D12Device* Devic
 #endif // PLATFORM_WIN64
 #endif // WITH_UE_AND_ORT_SUPPORT
 
+
+
 /* UNeuralNetwork public functions
  *****************************************************************************/
 
@@ -244,8 +238,8 @@ void UNeuralNetwork::FImplBackEndUEAndORT::WarnAndSetDeviceToCPUIfDX12NotEnabled
 			InOutDeviceType = ENeuralDeviceType::CPU;
 
 			const FString RHIName = GDynamicRHI->GetName();
-			const FString ErrorMessage = TEXT("On Windows, only DirectX 12 rendering (\"D3D12\") is compatible with the UEAndORT back end of NeuralNetworkInference (NNI). Instead, \"") + RHIName
-				+ TEXT("\" was used. You have the following options:\n\n"
+			const FString ErrorMessage = TEXT("On Windows, only DirectX 12 rendering (\"D3D12\") is compatible with the UEAndORT back end of NeuralNetworkInference (NNI). Instead, \"")
+				+ RHIName + TEXT("\" was used. You have the following options:\n\n"
 					"\t1. (Recommended) Switch Unreal Engine to DX12. In order to do that:\n"
 					"\t\t - Go to \"Project Settings\", \"Platforms\", \"Windows\", \"Default RHI\".\n"
 					"\t\t - Select \"DirectX 12\".\n"
@@ -275,83 +269,10 @@ bool UNeuralNetwork::FImplBackEndUEAndORT::IsGPUConfigCompatible()
 	return true;
 }
 
-//bool UNeuralNetwork::FImplBackEndUEAndORT::Load(TSharedPtr<FImplBackEndUEAndORT>& InOutImplBackEndUEAndORT, TArray<bool>& OutAreInputTensorSizesVariable, const TArray<uint8>& InModelReadFromFileInBytes, const FString& InModelFullFilePath, const ENeuralDeviceType InDeviceType)
-//{
-//	// Initialize and configure InOutImplBackEndUEAndORT
-//	const FRedirectCoutAndCerrToUeLog RedirectCoutAndCerrToUeLog;
-//	if (!InitializedAndConfigureMembers(InOutImplBackEndUEAndORT, InModelFullFilePath, InDeviceType))
-//	{
-//		UE_LOG(LogNeuralNetworkInference, Warning, TEXT("FImplBackEndUEAndORT::Load(): InitializedAndConfigureMembers failed."));
-//		return false;
-//	}
-//
-//	// Note: This code uses the changes in the ONNX Runtime API, which are not needed for desktop platforms (where ONNX/ProtoBuf is supported)
-//	// Fill InModelReadFromFileInBytes from ONNX/ORT file
-//	const FString FileExtension = FPaths::GetExtension(InModelFullFilePath, /*bIncludeDot*/ false);
-//	const char* const FilePathCharPtr = TCHAR_TO_ANSI(*InModelFullFilePath);
-//	// If ONNX file, turn into ORT format first
-//	if (FileExtension.Equals(TEXT("onnx"), ESearchCase::IgnoreCase))
-//	{
-//		FString ONNXPathPart, ONNXFilenamePart, ONNXExtensionPart;
-//		FPaths::Split(InModelFullFilePath, ONNXPathPart, ONNXFilenamePart, ONNXExtensionPart);
-//		const FString OutputORTOptimizedModelPath = ONNXPathPart / ONNXFilenamePart + TEXT(".ort");
-//#ifdef _WIN32
-//		InOutImplBackEndUEAndORT->SessionOptions->SetOptimizedModelFilePath(*OutputORTOptimizedModelPath);
-//#else
-//		InOutImplBackEndUEAndORT->SessionOptions->SetOptimizedModelFilePath(TCHAR_TO_ANSI(*OutputORTOptimizedModelPath));
-//#endif //_WIN32
-//		// ONNX --> ORT conversion
-//		// This session is just temporarily opened so the ORT model file can be generated
-//		InOutImplBackEndUEAndORT->Session = MakeUnique<Ort::Session>(*InOutImplBackEndUEAndORT->Environment,
-//#ifdef _WIN32
-//			*InModelFullFilePath,
-//#else
-//			FilePathCharPtr,
-//#endif
-//			*InOutImplBackEndUEAndORT->SessionOptions);
-//	
-//		// Read model from OutputORTOptimizedModelPath
-//		return Load(OutputORTOptimizedModelPath);
-//	}
-//	// Create session (it should be ORT file at this point), and read InModelReadFromFileInBytes if not empty
-//	else if (FileExtension.Equals(TEXT("ort"), ESearchCase::IgnoreCase))
-//	{
-//		// Read model from InModelFullFilePath
-//		std::vector<uint8_t> OutputModelReadFromFileInBytesVector;
-//		InOutImplBackEndUEAndORT->Session = MakeUnique<Ort::Session>(*InOutImplBackEndUEAndORT->Environment,
-//#ifdef _WIN32
-//			*InModelFullFilePath,
-//#else
-//			FilePathCharPtr,
-//#endif
-//			*InOutImplBackEndUEAndORT->SessionOptions, &OutputModelReadFromFileInBytesVector);
-//	
-//		// Fill InModelReadFromFileInBytes
-//		const int32 ArraySize = OutputModelReadFromFileInBytesVector.size();
-//		if (ArraySize > 0)
-//		{
-//			InModelReadFromFileInBytes.SetNumUninitialized(ArraySize);
-//			FMemory::Memcpy(InModelReadFromFileInBytes.GetData(), &OutputModelReadFromFileInBytesVector[0], ArraySize);
-//		}
-//	
-//		return Load();
-//	}
-//	// Else
-//	UE_LOG(LogNeuralNetworkInference, Warning, TEXT("FImplBackEndUEAndORT::Load(): Unknown file format \"%s\" used."), *FileExtension);
-//	return false;
-//}
-
-bool UNeuralNetwork::FImplBackEndUEAndORT::Load(
-			TSharedPtr<FImplBackEndUEAndORT>& InOutImplBackEndUEAndORT, 
-			FOnAsyncRunCompleted& InOutOnAsyncRunCompletedDelegate, 
-			std::atomic<bool>& bInOutIsBackgroundThreadRunning,
-			FCriticalSection& InOutResoucesCriticalSection, 
-			TArray<bool>& OutAreInputTensorSizesVariable, 
-			const TArray<uint8>& InModelReadFromFileInBytes, 
-			const FString& InModelFullFilePath,
-			const ENeuralDeviceType InDeviceType, 
-			const ENeuralDeviceType InInputDeviceType, 
-			const ENeuralDeviceType InOutputDeviceType)
+bool UNeuralNetwork::FImplBackEndUEAndORT::Load(TSharedPtr<FImplBackEndUEAndORT>& InOutImplBackEndUEAndORT,
+	FOnAsyncRunCompleted& InOutOnAsyncRunCompletedDelegate, std::atomic<bool>& bInOutIsBackgroundThreadRunning, FCriticalSection& InOutResoucesCriticalSection,
+	TArray<bool>& OutAreInputTensorSizesVariable, const TArray<uint8>& InModelReadFromFileInBytes, const FString& InModelFullFilePath,
+	const ENeuralDeviceType InDeviceType, const ENeuralDeviceType InInputDeviceType, const ENeuralDeviceType InOutputDeviceType)
 {
 #ifdef WITH_UE_AND_ORT_SUPPORT
 #if WITH_EDITOR
@@ -360,20 +281,15 @@ bool UNeuralNetwork::FImplBackEndUEAndORT::Load(
 	{
 		// Avoid multi-threaded crashes
 		const FRedirectCoutAndCerrToUeLog RedirectCoutAndCerrToUeLog;
-		
+
 		if (InOutImplBackEndUEAndORT.IsValid())
 		{
 			InOutImplBackEndUEAndORT->IsAsyncTaskDone();
 		}
 
 		// Initialize and configure InOutImplBackEndUEAndORT
-		if (!UNeuralNetwork::FImplBackEndUEAndORT::InitializedAndConfigureMembers(
-					InOutImplBackEndUEAndORT, 
-					InOutOnAsyncRunCompletedDelegate, 
-					bInOutIsBackgroundThreadRunning, 
-					InOutResoucesCriticalSection, 
-					InModelFullFilePath, 
-					InDeviceType))
+		if (!UNeuralNetwork::FImplBackEndUEAndORT::InitializedAndConfigureMembers(InOutImplBackEndUEAndORT, InOutOnAsyncRunCompletedDelegate,
+				bInOutIsBackgroundThreadRunning, InOutResoucesCriticalSection, InModelFullFilePath, InDeviceType))
 		{
 			UE_LOG(LogNeuralNetworkInference, Warning, TEXT("FImplBackEndUEAndORT::Load(): InitializedAndConfigureMembers failed."));
 			return false;
@@ -383,7 +299,8 @@ bool UNeuralNetwork::FImplBackEndUEAndORT::Load(
 		if (InModelReadFromFileInBytes.Num() > 0)
 		{
 			// Read model from ModelReadFromFileInBytesVector
-			InOutImplBackEndUEAndORT->Session = MakeUnique<Ort::Session>(*InOutImplBackEndUEAndORT->Environment, InModelReadFromFileInBytes.GetData(), InModelReadFromFileInBytes.Num(), *InOutImplBackEndUEAndORT->SessionOptions);
+			InOutImplBackEndUEAndORT->Session = MakeUnique<Ort::Session>(*InOutImplBackEndUEAndORT->Environment, InModelReadFromFileInBytes.GetData(),
+				InModelReadFromFileInBytes.Num(), *InOutImplBackEndUEAndORT->SessionOptions);
 
 #ifdef PLATFORM_WIN64
 			// Check if resource allocator is properly initialized
@@ -406,12 +323,14 @@ bool UNeuralNetwork::FImplBackEndUEAndORT::Load(
 		ENeuralDeviceType OutputDeviceType = InOutputDeviceType;
 		if (InDeviceType == ENeuralDeviceType::CPU && (InInputDeviceType == ENeuralDeviceType::GPU || InOutputDeviceType == ENeuralDeviceType::GPU))
 		{
-			UE_LOG(LogNeuralNetworkInference, Warning, TEXT("FImplBackEndUEAndORT::Load(): DeviceType is CPU but Input and/or Output is set to GPU, setting all to CPU."));
+			UE_LOG(LogNeuralNetworkInference, Warning,
+				TEXT("FImplBackEndUEAndORT::Load(): DeviceType is CPU but Input and/or Output is set to GPU, setting all to CPU."));
 			InputDeviceType = ENeuralDeviceType::CPU;
 			OutputDeviceType = ENeuralDeviceType::CPU;
 		}
 
-		if (!InOutImplBackEndUEAndORT->ConfigureTensors(InOutImplBackEndUEAndORT->InputTensors, &OutAreInputTensorSizesVariable, InDeviceType, InputDeviceType, OutputDeviceType))
+		if (!InOutImplBackEndUEAndORT->ConfigureTensors(InOutImplBackEndUEAndORT->InputTensors, &OutAreInputTensorSizesVariable, InDeviceType,
+			InputDeviceType, OutputDeviceType))
 		{
 			UE_LOG(LogNeuralNetworkInference, Warning, TEXT("FImplBackEndUEAndORT::Load(): Failed to configure input tensors."));
 			return false;
@@ -423,7 +342,7 @@ bool UNeuralNetwork::FImplBackEndUEAndORT::Load(
 			return false;
 		}
 
-		// Initializing AsyncTask 
+		// Initializing AsyncTask
 		InOutImplBackEndUEAndORT->NeuralNetworkAsyncTask = MakeUnique<FAsyncTask<FNeuralNetworkAsyncTask>>(InOutImplBackEndUEAndORT.Get());
 		
 		return true;
@@ -437,13 +356,14 @@ bool UNeuralNetwork::FImplBackEndUEAndORT::Load(
 #endif //WITH_EDITOR
 
 #else //WITH_UE_AND_ORT_SUPPORT
-	UE_LOG(LogNeuralNetworkInference, Warning, TEXT("FImplBackEndUEAndORT::Load(): Platform or Operating System not suported yet for UEAndORT BackEnd. Set BackEnd to ENeuralBackEnd::Auto (recommended) or ENeuralBackEnd::UEOnly for this platform."));
+	UE_LOG(LogNeuralNetworkInference, Warning, TEXT("FImplBackEndUEAndORT::Load(): Platform or Operating System not suported yet for UEAndORT"
+		" BackEnd. Set BackEnd to ENeuralBackEnd::Auto (recommended) or ENeuralBackEnd::UEOnly for this platform."));
 	return false;
 #endif //WITH_UE_AND_ORT_SUPPORT
 }
 
 #ifdef WITH_UE_AND_ORT_SUPPORT
-void UNeuralNetwork::FImplBackEndUEAndORT::IsAsyncTaskDone() const 
+void UNeuralNetwork::FImplBackEndUEAndORT::IsAsyncTaskDone() const
 {
 	if (NeuralNetworkAsyncTask && !NeuralNetworkAsyncTask->IsDone())
 	{
@@ -457,7 +377,7 @@ void UNeuralNetwork::FImplBackEndUEAndORT::ClearResources()
 {
 #ifdef PLATFORM_WIN64
 	if (DmlGPUAllocator.IsValid() && DmlGPUAllocator->IsValid())
-	{ 
+	{
 		const int32 Num = DmlGPUResources.Num();
 		for (int32 Index = 0; Index < Num; ++Index)
 		{
@@ -470,7 +390,8 @@ void UNeuralNetwork::FImplBackEndUEAndORT::ClearResources()
 #endif //PLATFORM_WIN64
 }
 
-UNeuralNetwork::FImplBackEndUEAndORT::FImplBackEndUEAndORT(FOnAsyncRunCompleted& InOutOnAsyncRunCompletedDelegate, std::atomic<bool>& bInIsBackgroundThreadRunning, FCriticalSection& InResoucesCriticalSection)
+UNeuralNetwork::FImplBackEndUEAndORT::FImplBackEndUEAndORT(FOnAsyncRunCompleted& InOutOnAsyncRunCompletedDelegate,
+	std::atomic<bool>& bInIsBackgroundThreadRunning, FCriticalSection& InResoucesCriticalSection)
 	: OnAsyncRunCompletedDelegate(InOutOnAsyncRunCompletedDelegate)
 	, bIsBackgroundThreadRunning(bInIsBackgroundThreadRunning)
 	, ResoucesCriticalSection(InResoucesCriticalSection)
@@ -479,7 +400,8 @@ UNeuralNetwork::FImplBackEndUEAndORT::FImplBackEndUEAndORT(FOnAsyncRunCompleted&
 
 #endif //WITH_UE_AND_ORT_SUPPORT
 
-void UNeuralNetwork::FImplBackEndUEAndORT::Run(const ENeuralNetworkSynchronousMode InSynchronousMode, const ENeuralDeviceType InDeviceType, const ENeuralDeviceType InInputDeviceType, const ENeuralDeviceType InOutputDeviceType)
+void UNeuralNetwork::FImplBackEndUEAndORT::Run(const ENeuralNetworkSynchronousMode InSynchronousMode, const ENeuralDeviceType InDeviceType,
+	const ENeuralDeviceType InInputDeviceType, const ENeuralDeviceType InOutputDeviceType)
 {
 #ifdef WITH_UE_AND_ORT_SUPPORT
 #if WITH_EDITOR
@@ -514,7 +436,8 @@ void UNeuralNetwork::FImplBackEndUEAndORT::Run(const ENeuralNetworkSynchronousMo
 #endif //WITH_EDITOR
 
 #else //WITH_UE_AND_ORT_SUPPORT
-	UE_LOG(LogNeuralNetworkInference, Warning, TEXT("FImplBackEndUEAndORT::Run(): Platform or Operating System not suported yet for UEAndORT BackEnd. Set BackEnd to ENeuralBackEnd::Auto or ENeuralBackEnd::UEOnly for this platform."));
+	UE_LOG(LogNeuralNetworkInference, Warning, TEXT("FImplBackEndUEAndORT::Run(): Platform or Operating System not suported yet for UEAndORT"
+		" BackEnd. Set BackEnd to ENeuralBackEnd::Auto or ENeuralBackEnd::UEOnly for this platform."));
 #endif //WITH_UE_AND_ORT_SUPPORT
 }
 
@@ -523,7 +446,8 @@ void UNeuralNetwork::FImplBackEndUEAndORT::Run(const ENeuralNetworkSynchronousMo
 
 #ifdef WITH_UE_AND_ORT_SUPPORT
 
-void UNeuralNetwork::FImplBackEndUEAndORT::RunSessionAsync(const ENeuralDeviceType InDeviceType, const ENeuralDeviceType InInputDeviceType, const ENeuralDeviceType InOutputDeviceType)
+void UNeuralNetwork::FImplBackEndUEAndORT::RunSessionAsync(const ENeuralDeviceType InDeviceType, const ENeuralDeviceType InInputDeviceType,
+	const ENeuralDeviceType InOutputDeviceType)
 {
 	const FScopeLock ResourcesLock(&ResoucesCriticalSection);
 
@@ -533,7 +457,8 @@ void UNeuralNetwork::FImplBackEndUEAndORT::RunSessionAsync(const ENeuralDeviceTy
 	bIsBackgroundThreadRunning = false;
 }
 
-void UNeuralNetwork::FImplBackEndUEAndORT::RunSessionSync(const ENeuralDeviceType InDeviceType, const ENeuralDeviceType InInputDeviceType, const ENeuralDeviceType InOutputDeviceType)
+void UNeuralNetwork::FImplBackEndUEAndORT::RunSessionSync(const ENeuralDeviceType InDeviceType, const ENeuralDeviceType InInputDeviceType,
+	const ENeuralDeviceType InOutputDeviceType)
 {
 	RunSessionImpl(InDeviceType, InInputDeviceType, InOutputDeviceType);
 }
@@ -545,7 +470,8 @@ BEGIN_SHADER_PARAMETER_STRUCT(FUploadTensorParameters, )
 	RDG_BUFFER_ACCESS(Input, ERHIAccess::CopyDest)
 END_SHADER_PARAMETER_STRUCT()
 
-void UNeuralNetwork::FImplBackEndUEAndORT::RunSessionImpl(const ENeuralDeviceType InDeviceType, const ENeuralDeviceType InInputDeviceType, const ENeuralDeviceType InOutputDeviceType)
+void UNeuralNetwork::FImplBackEndUEAndORT::RunSessionImpl(const ENeuralDeviceType InDeviceType, const ENeuralDeviceType InInputDeviceType,
+	const ENeuralDeviceType InOutputDeviceType)
 {
 	if (InDeviceType == ENeuralDeviceType::GPU)
 	{
@@ -622,21 +548,23 @@ void UNeuralNetwork::FImplBackEndUEAndORT::RunSessionImpl(const ENeuralDeviceTyp
 }
 
 bool UNeuralNetwork::FImplBackEndUEAndORT::InitializedAndConfigureMembers(
-			TSharedPtr<FImplBackEndUEAndORT>& InOutImplBackEndUEAndORT, 
+			TSharedPtr<FImplBackEndUEAndORT>& InOutImplBackEndUEAndORT,
 			FOnAsyncRunCompleted& InOutOnAsyncRunCompletedDelegate,
 			std::atomic<bool>& bInOutIsBackgroundThreadRunning,
-			FCriticalSection& InOutResoucesCriticalSection, 
-			const FString& InModelFullFilePath, 
+			FCriticalSection& InOutResoucesCriticalSection,
+			const FString& InModelFullFilePath,
 			const ENeuralDeviceType InDeviceType)
 {
 	// Initialize InOutImplBackEndUEAndORT
 	if (!InOutImplBackEndUEAndORT.IsValid())
 	{
-		InOutImplBackEndUEAndORT = MakeShared<FImplBackEndUEAndORT>(InOutOnAsyncRunCompletedDelegate, bInOutIsBackgroundThreadRunning, InOutResoucesCriticalSection);
+		InOutImplBackEndUEAndORT = MakeShared<FImplBackEndUEAndORT>(InOutOnAsyncRunCompletedDelegate, bInOutIsBackgroundThreadRunning,
+			InOutResoucesCriticalSection);
 
 		// Set up ORT and create an environment
 		const char* const ModelFullFilePathCharPtr = TCHAR_TO_ANSI(*InModelFullFilePath);
-		InOutImplBackEndUEAndORT->Environment = MakeUnique<Ort::Env>(ORT_LOGGING_LEVEL_WARNING, ModelFullFilePathCharPtr); // Any unique string would work, it does not need to be the file path
+		// @todo: ModelFullFilePathCharPtr -> I thought any unique string would work, but it might be output logging file, so it has to be a non-existing file!
+		InOutImplBackEndUEAndORT->Environment = MakeUnique<Ort::Env>(ORT_LOGGING_LEVEL_WARNING, ModelFullFilePathCharPtr);
 
 		InOutImplBackEndUEAndORT->Allocator = MakeUnique<Ort::AllocatorWithDefaultOptions>();
 
@@ -683,7 +611,8 @@ bool UNeuralNetwork::FImplBackEndUEAndORT::ConfigureMembers(const ENeuralDeviceT
 
 		if (Rhi->GetNumAdapters() > 1 || Rhi->GetAdapter(0).GetDesc().NumDeviceNodes > 1)
 		{
-			UE_LOG(LogNeuralNetworkInference, Warning, TEXT("FImplBackEndUEAndORT::ConfigureMembers(): There are multiple (%d) adapters and/or multiple (%d) devices, using device at index 0."),
+			UE_LOG(LogNeuralNetworkInference, Warning,
+				TEXT("FImplBackEndUEAndORT::ConfigureMembers(): There are multiple (%d) adapters and/or multiple (%d) devices, using device at index 0."),
 				Rhi->GetNumAdapters(), Rhi->GetAdapter(0).GetDesc().NumDeviceNodes);
 			return false;
 		}
@@ -714,14 +643,17 @@ bool UNeuralNetwork::FImplBackEndUEAndORT::ConfigureMembers(const ENeuralDeviceT
 		DmlProviderOptions->cmd_queue = NativeCmdQ;
 		DmlProviderOptions->resource_allocator = DmlGPUAllocator->GetAllocatorAddressOf();
 
-		if (OrtSessionOptionsAppendExecutionProviderWithOptions_DML(*SessionOptions, DmlProviderOptions.Get())) // OrtSessionOptionsAppendExecutionProvider_DML(*SessionOptions, 0) without sharing
+		// OrtSessionOptionsAppendExecutionProvider_DML(*SessionOptions, 0) without sharing
+		if (OrtSessionOptionsAppendExecutionProviderWithOptions_DML(*SessionOptions, DmlProviderOptions.Get()))
 		{
-			UE_LOG(LogNeuralNetworkInference, Warning, TEXT("FImplBackEndUEAndORT::ConfigureMembers(): Some error occurred when using OrtSessionOptionsAppendExecutionProviderEx_DML()."));
+			UE_LOG(LogNeuralNetworkInference, Warning,
+				TEXT("FImplBackEndUEAndORT::ConfigureMembers(): Some error occurred when using OrtSessionOptionsAppendExecutionProviderEx_DML()."));
 			return false;
 		}
 		return true; // @todo: Remove this line when NNI_HLSL is working
 #else
-		UE_LOG(LogNeuralNetworkInference, Warning, TEXT("FImplBackEndUEAndORT::ConfigureMembers(): GPU mode only supported in Windows for now. Please, switch to CPU or to Windows."));
+		UE_LOG(LogNeuralNetworkInference, Warning,
+			TEXT("FImplBackEndUEAndORT::ConfigureMembers(): GPU mode only supported in Windows for now. Please, switch to CPU or to Windows."));
 
 		//SessionOptions->SetGraphOptimizationLevel(GraphOptimizationLevel::ORT_DISABLE_ALL); // ORT_ENABLE_ALL, ORT_ENABLE_EXTENDED, ORT_ENABLE_BASIC, ORT_DISABLE_ALL
 		//if (OrtSessionOptionsAppendExecutionProvider_NNI_HLSL(*SessionOptions, 0))
@@ -751,7 +683,8 @@ bool UNeuralNetwork::FImplBackEndUEAndORT::ConfigureMembers(const ENeuralDeviceT
 	return true;
 }
 
-bool UNeuralNetwork::FImplBackEndUEAndORT::ConfigureTensors(TArray<FNeuralTensor>& OutTensors, TArray<bool>* OutAreInputTensorSizesVariable, const ENeuralDeviceType InDeviceType, const ENeuralDeviceType InInputDeviceType, const ENeuralDeviceType InOutputDeviceType)
+bool UNeuralNetwork::FImplBackEndUEAndORT::ConfigureTensors(TArray<FNeuralTensor>& OutTensors, TArray<bool>* OutAreInputTensorSizesVariable,
+	const ENeuralDeviceType InDeviceType, const ENeuralDeviceType InInputDeviceType, const ENeuralDeviceType InOutputDeviceType)
 {
 	const bool bIsInput = (OutAreInputTensorSizesVariable != nullptr);
 	TArray<const char*> TensorNames;
@@ -803,7 +736,9 @@ bool UNeuralNetwork::FImplBackEndUEAndORT::ConfigureTensors(TArray<FNeuralTensor
 			else
 			{
 				TensorDataType = ENeuralDataType::None;
-				UE_LOG(LogNeuralNetworkInference, Warning, TEXT("FImplBackEndUEAndORT::ConfigureTensors(): ONNXTensorElementDataTypeEnum = %d not implemented yet."), (int32)ONNXTensorElementDataTypeEnum);
+				UE_LOG(LogNeuralNetworkInference, Warning,
+					TEXT("FImplBackEndUEAndORT::ConfigureTensors(): ONNXTensorElementDataTypeEnum = %d not implemented yet."),
+					(int32)ONNXTensorElementDataTypeEnum);
 				return false;
 			}
 		}
@@ -854,7 +789,7 @@ bool UNeuralNetwork::FImplBackEndUEAndORT::ConfigureTensors(TArray<FNeuralTensor
 		{
 			TensorGPUType = ENeuralTensorTypeGPU::Generic;
 		}
-		
+
 		TensorGPUTypes.Push(TensorGPUType);
 
 		CurrentTypeInfo.release();
@@ -863,13 +798,14 @@ bool UNeuralNetwork::FImplBackEndUEAndORT::ConfigureTensors(TArray<FNeuralTensor
 	return SetTensorsFromNetwork(OutTensors, TensorNames, TensorDataTypes, TensorSizes, TensorGPUTypes, bIsInput);
 }
 
-bool UNeuralNetwork::FImplBackEndUEAndORT::SetTensorsFromNetwork(TArray<FNeuralTensor>& OutTensors, TArray<const char*>& InTensorNames, TArray<ENeuralDataType>& InTensorDataTypes,
-	TArray<TArray<int64>>& InSizes, TArray<ENeuralTensorTypeGPU>& InTensorGPUTypes, const bool bIsInput)
+bool UNeuralNetwork::FImplBackEndUEAndORT::SetTensorsFromNetwork(TArray<FNeuralTensor>& OutTensors, TArray<const char*>& InTensorNames,
+	TArray<ENeuralDataType>& InTensorDataTypes, TArray<TArray<int64>>& InSizes, TArray<ENeuralTensorTypeGPU>& InTensorGPUTypes, const bool bIsInput)
 {
 	const int32 TensorNumber = InTensorNames.Num();
 	if (InTensorDataTypes.Num() != TensorNumber || InSizes.Num() != TensorNumber)
 	{
-		UE_LOG(LogNeuralNetworkInference, Warning, TEXT("FImplBackEndUEAndORT::SetTensorsFromNetwork(): InTensorNames.Num() == InTensorDataTypes.Num() == InSizes.Num() failed, %d vs. %d vs. %d."),
+		UE_LOG(LogNeuralNetworkInference, Warning,
+			TEXT("FImplBackEndUEAndORT::SetTensorsFromNetwork(): InTensorNames.Num() == InTensorDataTypes.Num() == InSizes.Num() failed, %d vs. %d vs. %d."),
 			InTensorNames.Num(), InTensorDataTypes.Num(), InSizes.Num());
 		return false;
 	}
@@ -878,7 +814,8 @@ bool UNeuralNetwork::FImplBackEndUEAndORT::SetTensorsFromNetwork(TArray<FNeuralT
 	TArray<const char*>& TensorNames = (bIsInput ? InputTensorNames : OutputTensorNames);
 	Swap(TensorNames, InTensorNames);
 
-	// Note: Switching from/to CPU to/from GPU would cause the FNeuralTensors to be re-initialized. We need to avoid that. For that, we will only re-allocate the tensors...
+	// Note: Switching from/to CPU to/from GPU would cause the FNeuralTensors to be re-initialized. We need to avoid that. For that,
+	// we will only re-allocate the tensors...
 	// - If bAreTensorsAlreadyCreatedWithRightNames == false, meaning tensors had not been created until now for this network.
 	// - And if the existing tensors have the right size, given that SetNumUninitialized() only re-allocates them if their size has changed.
 
@@ -945,14 +882,16 @@ bool UNeuralNetwork::FImplBackEndUEAndORT::SetTensorsFromNetwork(TArray<FNeuralT
 
 			if (!OutTensors[TensorIndex].InitPooledBuffer(&D3DResource))
 			{
-				UE_LOG(LogNeuralNetworkInference, Warning, TEXT("FImplBackEndUEAndORT::SetTensorsFromNetwork(): Failed to initialize pooled buffer"));
+				UE_LOG(LogNeuralNetworkInference, Warning,
+					TEXT("FImplBackEndUEAndORT::SetTensorsFromNetwork(): Failed to initialize pooled buffer"));
 				return false;
 			}
 
 			// Link tensor with ORT blob
 			if (!LinkTensorResourceToONNXRuntime(OutTensors[TensorIndex], OrtTensors[TensorIndex], D3DResource))
 			{
-				UE_LOG(LogNeuralNetworkInference, Warning, TEXT("FImplBackEndUEAndORT::SetTensorsFromNetwork(): Failed to link GPU resource to ONNX runtime"));
+				UE_LOG(LogNeuralNetworkInference, Warning,
+					TEXT("FImplBackEndUEAndORT::SetTensorsFromNetwork(): Failed to link GPU resource to ONNX runtime"));
 				return false;
 			}
 		}
@@ -967,7 +906,8 @@ bool UNeuralNetwork::FImplBackEndUEAndORT::SetTensorsFromNetwork(TArray<FNeuralT
 	return true;
 }
 
-void UNeuralNetwork::FImplBackEndUEAndORT::LinkTensorToONNXRuntime(TArray<FNeuralTensor>& InOutTensors, TArray<Ort::Value>& InOutOrtTensors, Ort::MemoryInfo& InOutAllocatorInfo, const int32 InTensorIndex)
+void UNeuralNetwork::FImplBackEndUEAndORT::LinkTensorToONNXRuntime(TArray<FNeuralTensor>& InOutTensors, TArray<Ort::Value>& InOutOrtTensors,
+	Ort::MemoryInfo& InOutAllocatorInfo, const int32 InTensorIndex)
 {
 	const TArray<int64>& Sizes = InOutTensors[InTensorIndex].GetSizes();
 	if (Sizes.Num() > 0 && InOutTensors[InTensorIndex].Num() > 0)
@@ -1035,15 +975,18 @@ bool UNeuralNetwork::FImplBackEndUEAndORT::LinkTensorResourceToONNXRuntime(FNeur
 			SizesInt64t.SetNumUninitialized(ArrayDimensions);
 			FMemory::Memcpy(SizesInt64t.GetData(), (int64_t*)Sizes.GetData(), sizeof(int64_t) * ArrayDimensions);
 #endif //_WIN32
-			InOutOrtTensor = Ort::Value::CreateTensor(DmlGPUAllocator->GetProviderMemoryInfo(), DmlGPUAllocation, InOutTensor.NumInBytes(), SizesInt64t.GetData(), ArrayDimensions, ONNX_TENSOR_ELEMENT_DATA_TYPE_FLOAT);
+			InOutOrtTensor = Ort::Value::CreateTensor(DmlGPUAllocator->GetProviderMemoryInfo(), DmlGPUAllocation, InOutTensor.NumInBytes(), SizesInt64t.GetData(),
+				ArrayDimensions, ONNX_TENSOR_ELEMENT_DATA_TYPE_FLOAT);
 		}
 		//else if (NeuralDataType == ENeuralDataType::Double)
 		//{
-		//	InOutOrtTensor = Ort::Value::CreateTensor(DmlGPUAllocator->GetProviderMemoryInfo(), DmlGPUAllocation, InOutTensor.NumInBytes(), SizesInt64t.GetData(), ArrayDimensions, ONNX_TENSOR_ELEMENT_DATA_TYPE_DOUBLE);
+		//	InOutOrtTensor = Ort::Value::CreateTensor(DmlGPUAllocator->GetProviderMemoryInfo(), DmlGPUAllocation, InOutTensor.NumInBytes(), SizesInt64t.GetData(),
+		//		ArrayDimensions, ONNX_TENSOR_ELEMENT_DATA_TYPE_DOUBLE);
 		//}
 		else
 		{
-			UE_LOG(LogNeuralNetworkInference, Warning, TEXT("FImplBackEndUEAndORT::LinkTensorToONNXRuntime(): Not implemented (yet) for ENeuralDataType = %d."), (int32)NeuralDataType);
+			UE_LOG(LogNeuralNetworkInference, Warning,
+				TEXT("FImplBackEndUEAndORT::LinkTensorToONNXRuntime(): Not implemented (yet) for ENeuralDataType = %d."), (int32)NeuralDataType);
 			return false;
 		}
 	}
diff --git a/Engine/Plugins/Experimental/NeuralNetworkInference/Source/NeuralNetworkInferenceQA/Private/ModelUnitTester.cpp b/Engine/Plugins/Experimental/NeuralNetworkInference/Source/NeuralNetworkInferenceQA/Private/ModelUnitTester.cpp
--- a/Engine/Plugins/Experimental/NeuralNetworkInference/Source/NeuralNetworkInferenceQA/Private/ModelUnitTester.cpp
+++ b/Engine/Plugins/Experimental/NeuralNetworkInference/Source/NeuralNetworkInferenceQA/Private/ModelUnitTester.cpp
@@ -10,7 +10,7 @@
 /* FModelUnitTester axuiliary functions
  *****************************************************************************/
 
-struct FNNIUnitTesterTimeData 
+struct FNNIUnitTesterTimeData
 {
 	FNeuralStatsData ComputeTimeData;
 	FNeuralStatsData InputCopyTimeData;
@@ -23,38 +23,38 @@ struct FNNIUnitTesterTimeData
 	{}
 };
 
-FNNIUnitTesterTimeData FModelUnitTester_GetTimeInformation(UNeuralNetwork* InOutNetwork, TArray<float>& OutCPUGPUCPUOutput, const TArray<float>& InInputArray, const int32 InRepetitions) 
+FNNIUnitTesterTimeData FModelUnitTester_GetTimeInformation(UNeuralNetwork* InOutNetwork, TArray<float>& OutCPUGPUCPUOutput, const TArray<float>& InInputArray, const int32 InRepetitions)
 {
-	/* Input/output copy speed */ 
+	/* Input/output copy speed */
 	FNeuralStats OutCopyingStats;
 	FNeuralTimer Timer;
-	for (int32 TimerIndex = 0; TimerIndex < InRepetitions; ++TimerIndex) 
-	{ 
-		InOutNetwork->SetInputFromArrayCopy(InInputArray); 
+	for (int32 TimerIndex = 0; TimerIndex < InRepetitions; ++TimerIndex)
+	{
+		InOutNetwork->SetInputFromArrayCopy(InInputArray);
 		Timer.Tic();
 		OutCPUGPUCPUOutput = InOutNetwork->GetOutputTensor().GetArrayCopy<float>();
 		const float CurrentCopyingTime = Timer.Toc();
 		OutCopyingStats.StoreSample(CurrentCopyingTime);
 
-	} 
+	}
 
-	/* Forward() speed */ 
+	/* Forward() speed */
 	if (InRepetitions > 1)
-	{ 
-		for (int32 TimerIndex = 0; TimerIndex < 5; ++TimerIndex) 
-		{ 
-			InOutNetwork->Run(); 
-		} 
-	} 
-
-	if (InRepetitions > 0) 
-	{ 
-		for (int32 TimerIndex = 0; TimerIndex < InRepetitions; ++TimerIndex) 
-		{ 
+	{
+		for (int32 TimerIndex = 0; TimerIndex < 5; ++TimerIndex)
+		{
+			InOutNetwork->Run();
+		}
+	}
+
+	if (InRepetitions > 0)
+	{
+		for (int32 TimerIndex = 0; TimerIndex < InRepetitions; ++TimerIndex)
+		{
 			InOutNetwork->SetInputFromArrayCopy(InInputArray);
 			InOutNetwork->Run();
-			OutCPUGPUCPUOutput = InOutNetwork->GetOutputTensor().GetArrayCopy<float>(); 
-		} 
+			OutCPUGPUCPUOutput = InOutNetwork->GetOutputTensor().GetArrayCopy<float>();
+		}
 	}
 	// Return NetworkTimeData
 	return FNNIUnitTesterTimeData(InOutNetwork->GetInferenceStats(), InOutNetwork->GetInputMemoryTransferStats(), OutCopyingStats.GetStats());
@@ -70,41 +70,37 @@ bool FModelUnitTester::GlobalTest(const FString& InProjectContentDir, const FStr
 	// NOTE: models are separated into multiple lines to enable/disable particular network for faster testing
 
 	// Model names, input values, and number of repetitions for profiling
-	const TArray<FString> ModelNames(
-	{ 
-		TEXT("MLRigDeformer"), 
-		TEXT("cloth_network"), 
-		TEXT("HS"), 
-		TEXT("RL") 
+	const TArray<FString> ModelNames({
+		TEXT("MLRigDeformer"),
+		TEXT("cloth_network"),
+		TEXT("HS"),
+		TEXT("RL")
 	});
 
 	// Ground truths
-	const TArray<TArray<double>> CPUGroundTruths(
-	{ 
-		{3.728547, 0.008774, 4.595651, 212.193216, 742.434561, 4.250668, 4.717748}, 
+	const TArray<TArray<double>> CPUGroundTruths({
+		{3.728547, 0.008774, 4.595651, 212.193216, 742.434561, 4.250668, 4.717748},
 		{0.042571, 0.023693, 0.015783, 13.100505, 8.050994, 0.028807, 0.016387},
-		{138.372906, 126.753839, 127.287254, 130.316062, 127.303424, 124.800896, 126.546051}, 
-		{0.488662, 0.472437, 0.478862, 0.522685, 0.038322, 0.480848, 0.483821} 
+		{138.372906, 126.753839, 127.287254, 130.316062, 127.303424, 124.800896, 126.546051},
+		{0.488662, 0.472437, 0.478862, 0.522685, 0.038322, 0.480848, 0.483821}
 	});
 	
-	const TArray<TArray<double>> GPUGroundTruths(
-	{ 
-		{3.728547, 0.008774, 4.595651, 212.193208, 742.434578, 4.250668, 4.717748}, 
+	const TArray<TArray<double>> GPUGroundTruths({
+		{3.728547, 0.008774, 4.595651, 212.193208, 742.434578, 4.250668, 4.717748},
 		{0.042571, 0.023693, 0.015783, 13.100504, 8.050994, 0.028807, 0.016387},
-		{138.373184, 126.754100, 127.287398, 130.316194, 127.303495, 124.801134, 126.5462530}, 
-		{0.488662, 0.472437, 0.478862, 0.522685, 0.038322, 0.480848, 0.483821} 
+		{138.373184, 126.754100, 127.287398, 130.316194, 127.303495, 124.801134, 126.5462530},
+		{0.488662, 0.472437, 0.478862, 0.522685, 0.038322, 0.480848, 0.483821}
 	});
 
-	const TArray<float> InputArrayValues( // This one can be shorter than CPU/GPUGroundTruths
-	{ 
-		1.f, 
-		0.f, 
-		-1.f, 
-		100.f, 
-		-100.f, 
-		0.5f, 
-		-0.5f 
-	}); 
+	const TArray<float> InputArrayValues({ // This one can be shorter than CPU/GPUGroundTruths
+		1.f,
+		0.f,
+		-1.f,
+		100.f,
+		-100.f,
+		0.5f,
+		-0.5f
+	});
 	
 	// Speed profiling test - 0 repetitions means that test will not be run
 #ifdef WITH_UE_AND_ORT_SUPPORT
