UNN API simplified.
_codereview Francisco.Vicente
#rnx

diff --git a/Engine/Plugins/Experimental/NeuralNetworkInference/Source/NeuralNetworkInference/Private/NeuralNetwork.cpp b/Engine/Plugins/Experimental/NeuralNetworkInference/Source/NeuralNetworkInference/Private/NeuralNetwork.cpp
--- a/Engine/Plugins/Experimental/NeuralNetworkInference/Source/NeuralNetworkInference/Private/NeuralNetwork.cpp
+++ b/Engine/Plugins/Experimental/NeuralNetworkInference/Source/NeuralNetworkInference/Private/NeuralNetwork.cpp
@@ -209,18 +209,12 @@ void UNeuralNetwork::ResetStats()
 	InputMemoryTransferStatsModule.ResetStats();
 }
 
-
-bool UNeuralNetwork::IsGPUConfigCompatibleForUEAndORTBackEnd()
-{
-	return FImplBackEndUEAndORT::IsGPUConfigCompatible();
-}
-
-bool UNeuralNetwork::IsGPUConfigCompatibleForCurrentBackEnd() const
+bool UNeuralNetwork::IsGPUSupported() const
 {
 	// UEAndORT
 	if (BackEndForCurrentPlatform == ENeuralBackEnd::UEAndORT)
 	{
-		return IsGPUConfigCompatibleForUEAndORTBackEnd();
+		return FImplBackEndUEAndORT::IsGPUConfigCompatible();
 	}
 
 	// UEOnly
@@ -560,7 +554,7 @@ bool UNeuralNetwork::Load()
 	// UEAndORT
 	if (BackEndForCurrentPlatform == ENeuralBackEnd::UEAndORT)
 	{
-		UNeuralNetwork::FImplBackEndUEAndORT::WarnAndSetDeviceToCPUIfDX12NotEnabled(DeviceType);
+		UNeuralNetwork::FImplBackEndUEAndORT::WarnAndSetDeviceToCPUIfDX12NotEnabled(DeviceType, /*bShouldOpenMessageLog*/true);
 		bIsLoaded = UNeuralNetwork::FImplBackEndUEAndORT::Load(ImplBackEndUEAndORT, OnAsyncRunCompletedDelegate, bIsBackgroundThreadRunning, ResoucesCriticalSection, AreInputTensorSizesVariable,
 			ModelReadFromFileInBytes, ModelFullFilePath, GetDeviceType(), GetInputDeviceType(), GetOutputDeviceType());
 	}
@@ -661,10 +655,9 @@ void UNeuralNetwork::PostLoad()
 	if (ModelReadFromFileInBytes.Num() > 0)
 	{
 		// If GPU selected but not compatible, set to CPU
-		if (DeviceType == ENeuralDeviceType::GPU && !IsGPUConfigCompatibleForCurrentBackEnd())
+		if (BackEnd == ENeuralBackEnd::UEAndORT)
 		{
-			UE_LOG(LogNeuralNetworkInference, Warning, TEXT("UNeuralNetwork::PostLoad(): IsGPUConfigCompatibleForCurrentBackEnd() returned false, setting DeviceType to CPU."));
-			DeviceType = ENeuralDeviceType::CPU;
+			FImplBackEndUEAndORT::WarnAndSetDeviceToCPUIfDX12NotEnabled(DeviceType, /*bShouldOpenMessageLog*/false);
 		}
 		// Load
 		if (!Load())
diff --git a/Engine/Plugins/Experimental/NeuralNetworkInference/Source/NeuralNetworkInference/Private/NeuralNetworkImplBackEndUEAndORT.cpp b/Engine/Plugins/Experimental/NeuralNetworkInference/Source/NeuralNetworkInference/Private/NeuralNetworkImplBackEndUEAndORT.cpp
--- a/Engine/Plugins/Experimental/NeuralNetworkInference/Source/NeuralNetworkInference/Private/NeuralNetworkImplBackEndUEAndORT.cpp
+++ b/Engine/Plugins/Experimental/NeuralNetworkInference/Source/NeuralNetworkInference/Private/NeuralNetworkImplBackEndUEAndORT.cpp
@@ -35,7 +35,6 @@ NNI_THIRD_PARTY_INCLUDES_START
 #ifdef WITH_UE_AND_ORT_SUPPORT
 	#ifdef PLATFORM_WIN64
 	#include "core/providers/dml/dml_provider_factory.h"
-	//#include "core/framework/allocator.h"
 	#endif
 	#ifdef WITH_NNI_CPU_NOT_RECOMMENDED
 	#include "core/providers/nni_cpu/nni_cpu_provider_factory.h"
@@ -230,7 +229,7 @@ IDMLDevice* FPrivateImplBackEndUEAndORT::FDMLDeviceList::Add(ID3D12Device* Devic
 /* UNeuralNetwork public functions
  *****************************************************************************/
 
-void UNeuralNetwork::FImplBackEndUEAndORT::WarnAndSetDeviceToCPUIfDX12NotEnabled(ENeuralDeviceType& InOutDeviceType)
+void UNeuralNetwork::FImplBackEndUEAndORT::WarnAndSetDeviceToCPUIfDX12NotEnabled(ENeuralDeviceType& InOutDeviceType, const bool bInShouldOpenMessageLog)
 {
 	if (InOutDeviceType != ENeuralDeviceType::CPU)
 	{
@@ -245,12 +244,14 @@ void UNeuralNetwork::FImplBackEndUEAndORT::WarnAndSetDeviceToCPUIfDX12NotEnabled
 					"\t\t - Go to \"Project Settings\", \"Platforms\", \"Windows\", \"Default RHI\".\n"
 					"\t\t - Select \"DirectX 12\".\n"
 					"\t\t - Restart Unreal Engine.\n"
-					"\t2. Alternatively, switch the network to CPU with UNeuralNetwork::SetDeviceType().\n"
-					"\t3. (Not recommended) You could also switch the network to UEOnly with UNeuralNetwork::SetBackEnd().\n\n"
+					"\t2. Alternatively, switch the network to CPU with UNeuralNetwork::SetDeviceType().\n\n"
 					"Network set to CPU provisionally.");
 			UE_LOG(LogNeuralNetworkInference, Warning, TEXT("FImplBackEndUEAndORT::WarnAndSetDeviceToCPUIfDX12NotEnabled(): %s"), *ErrorMessage);
 #if WITH_EDITOR
-			FMessageDialog::Open(EAppMsgType::Ok, FText::FromString(ErrorMessage));
+			if (bInShouldOpenMessageLog)
+			{
+				FMessageDialog::Open(EAppMsgType::Ok, FText::FromString(ErrorMessage));
+			}
 #endif //WITH_EDITOR
 		}
 	}
diff --git a/Engine/Plugins/Experimental/NeuralNetworkInference/Source/NeuralNetworkInference/Private/NeuralNetworkImplBackEndUEAndORT.h b/Engine/Plugins/Experimental/NeuralNetworkInference/Source/NeuralNetworkInference/Private/NeuralNetworkImplBackEndUEAndORT.h
--- a/Engine/Plugins/Experimental/NeuralNetworkInference/Source/NeuralNetworkInference/Private/NeuralNetworkImplBackEndUEAndORT.h
+++ b/Engine/Plugins/Experimental/NeuralNetworkInference/Source/NeuralNetworkInference/Private/NeuralNetworkImplBackEndUEAndORT.h
@@ -97,7 +97,7 @@ public:
 	TArray<FNeuralTensor> InputTensors;
 	TArray<FNeuralTensor> OutputTensors;
 
-	static void WarnAndSetDeviceToCPUIfDX12NotEnabled(ENeuralDeviceType& InOutDeviceType);
+	static void WarnAndSetDeviceToCPUIfDX12NotEnabled(ENeuralDeviceType& InOutDeviceType, const bool bInShouldOpenMessageLog);
 
 	static bool IsGPUConfigCompatible();
 
diff --git a/Engine/Plugins/Experimental/NeuralNetworkInference/Source/NeuralNetworkInference/Public/NeuralNetwork.h b/Engine/Plugins/Experimental/NeuralNetworkInference/Source/NeuralNetworkInference/Public/NeuralNetwork.h
--- a/Engine/Plugins/Experimental/NeuralNetworkInference/Source/NeuralNetworkInference/Public/NeuralNetwork.h
+++ b/Engine/Plugins/Experimental/NeuralNetworkInference/Source/NeuralNetworkInference/Public/NeuralNetwork.h
@@ -131,12 +131,11 @@ public:
 	bool SetBackEnd(const ENeuralBackEnd InBackEnd);
 	
 	/**
-	 * IsGPUConfigCompatible will always return true for ENeuralBackEnd::UEOnly. For ENeuralBackEnd::UEAndORT, IsGPUConfigCompatible() and IsGPUConfigCompatibleForUEAndORTBackEnd() will return the same value:
+	 * IsGPUCompatible will always return true for ENeuralBackEnd::UEOnly. For ENeuralBackEnd::UEAndORT, it will return:
 	 * - True if DX12 is enabled, meaning UEAndORT can run on both the CPU and GPU. Also true if the current platform is not Windows.
 	 * - False if DX12 is disabled, meaning UEAndORT can run only run on the CPU. The user will need to enable DX12 to be able to run GPU, switch to CPU, or switch to the UEOnly back end.
 	 */
-	static bool IsGPUConfigCompatibleForUEAndORTBackEnd();
-	bool IsGPUConfigCompatibleForCurrentBackEnd() const;
+	bool IsGPUSupported() const;
 
 	/**
 	 * Functions to get/fill input.
diff --git a/Engine/Plugins/Experimental/NeuralNetworkInference/Source/NeuralNetworkInferenceQA/Private/ModelUnitTester.cpp b/Engine/Plugins/Experimental/NeuralNetworkInference/Source/NeuralNetworkInferenceQA/Private/ModelUnitTester.cpp
--- a/Engine/Plugins/Experimental/NeuralNetworkInference/Source/NeuralNetworkInferenceQA/Private/ModelUnitTester.cpp
+++ b/Engine/Plugins/Experimental/NeuralNetworkInference/Source/NeuralNetworkInferenceQA/Private/ModelUnitTester.cpp
@@ -205,8 +205,13 @@ bool FModelUnitTester::ModelLoadAccuracyAndSpeedTests(const FString& InProjectCo
 		// UEAndORT (if WITH_UE_AND_ORT_SUPPORT)
 #ifdef WITH_UE_AND_ORT_SUPPORT
 		bDidGlobalTestPassed &= ModelSpeedTest(UAssetModelFilePath, ENeuralDeviceType::CPU, ENeuralBackEnd::UEAndORT, InCPURepetitionsForUEAndORTBackEnd[ModelIndex]);
-#ifdef PLATFORM_WIN64
-		if (UNeuralNetwork::IsGPUConfigCompatibleForUEAndORTBackEnd())
+		UNeuralNetwork* Network = NetworkUassetLoadTest(UAssetModelFilePath);
+		if (Network->GetBackEndForCurrentPlatform() != ENeuralBackEnd::UEAndORT)
+		{
+			UE_LOG(LogNeuralNetworkInferenceQA, Warning, TEXT("-------------------- Default UAsset BackEnd should be UEAndORT."));
+			return false;
+		}
+		else if (Network->IsGPUSupported())
 		{
 			bDidGlobalTestPassed &= ModelSpeedTest(UAssetModelFilePath, ENeuralDeviceType::GPU, ENeuralBackEnd::UEAndORT, InGPURepetitionsForUEAndORTBackEnd[ModelIndex]);
 		}
@@ -214,9 +219,6 @@ bool FModelUnitTester::ModelLoadAccuracyAndSpeedTests(const FString& InProjectCo
 		{
 			UE_LOG(LogNeuralNetworkInferenceQA, Display, TEXT("-------------------- ModelSpeedTest-UEAndORT-GPU skipped (DX12 not enabled)."));
 		}
-#else //PLATFORM_WIN64
-		bDidGlobalTestPassed &= ModelSpeedTest(UAssetModelFilePath, ENeuralDeviceType::GPU, ENeuralBackEnd::UEAndORT, 0);
-#endif //PLATFORM_WIN64
 #endif //WITH_UE_AND_ORT_SUPPORT
 
 		// UEOnly
@@ -254,11 +256,6 @@ UNeuralNetwork* FModelUnitTester::NetworkUassetLoadTest(const FString& InUAssetP
 		ensureMsgf(false, TEXT("UNeuralNetwork is a nullptr. Path: \"%s\"."), *InUAssetPath);
 		return nullptr;
 	}
-	if (!UNeuralNetwork::IsGPUConfigCompatibleForUEAndORTBackEnd())
-	{
-		UE_LOG(LogNeuralNetworkInferenceQA, Display, TEXT("NetworkUassetLoadTest-UEAndORT: Set to CPU (DX12 not enabled)."));
-		Network->SetDeviceType(ENeuralDeviceType::CPU);
-	}
 	if (!Network->IsLoaded())
 	{
 		ensureMsgf(false, TEXT("UNeuralNetwork could not be loaded from uasset disk location. Path: \"%s\"."), *InUAssetPath);
@@ -276,11 +273,6 @@ UNeuralNetwork* FModelUnitTester::NetworkONNXOrORTLoadTest(const FString& InMode
 		ensureMsgf(false, TEXT("UNeuralNetwork is a nullptr. Path: \"%s\"."), *InModelFilePath);
 		return nullptr;
 	}
-	if (!UNeuralNetwork::IsGPUConfigCompatibleForUEAndORTBackEnd())
-	{
-		UE_LOG(LogNeuralNetworkInferenceQA, Display, TEXT("NetworkONNXOrORTLoadTest-UEAndORT: Set to CPU (DX12 not enabled)."));
-		Network->SetDeviceType(ENeuralDeviceType::CPU);
-	}
 	if (!Network->Load(InModelFilePath))
 	{
 		ensureMsgf(false, TEXT("UNeuralNetwork could not be loaded from ONNX file disk location. Path: \"%s\"."), *InModelFilePath);
@@ -344,7 +336,7 @@ bool FModelUnitTester::ModelAccuracyTest(UNeuralNetwork* InOutNetwork, const ENe
 		CPUOutputs.Emplace(InOutNetwork->GetOutputTensor().GetArrayCopy<float>());
 	}
 	
-	if (UNeuralNetwork::IsGPUConfigCompatibleForUEAndORTBackEnd())
+	if (InOutNetwork->IsGPUSupported())
 	{
 		// Input CPU + Network GPU + Output CPU
 		for (int32 Index = 0; Index < InputArrays.Num(); ++Index)
@@ -382,9 +374,9 @@ bool FModelUnitTester::ModelAccuracyTest(UNeuralNetwork* InOutNetwork, const ENe
 	
 	for (int32 Index = 0; Index < InputArrays.Num(); ++Index)
 	{
-		const TArray<float>& GPUGPUCPUOutput = (UNeuralNetwork::IsGPUConfigCompatibleForUEAndORTBackEnd() ? GPUGPUCPUOutputs[Index] : CPUOutputs[Index]);
-		const TArray<float>& CPUGPUCPUOutput = (UNeuralNetwork::IsGPUConfigCompatibleForUEAndORTBackEnd() ? CPUGPUCPUOutputs[Index] : CPUOutputs[Index]);
-		const TArray<float>& CPUGPUGPUOutput = (UNeuralNetwork::IsGPUConfigCompatibleForUEAndORTBackEnd() ? CPUGPUGPUOutputs[Index] : CPUOutputs[Index]);
+		const TArray<float>& GPUGPUCPUOutput = (InOutNetwork->IsGPUSupported() ? GPUGPUCPUOutputs[Index] : CPUOutputs[Index]);
+		const TArray<float>& CPUGPUCPUOutput = (InOutNetwork->IsGPUSupported() ? CPUGPUCPUOutputs[Index] : CPUOutputs[Index]);
+		const TArray<float>& CPUGPUGPUOutput = (InOutNetwork->IsGPUSupported() ? CPUGPUGPUOutputs[Index] : CPUOutputs[Index]);
 		// Prepare verbose
 		const double CPUAvgL1Norm = GetAveragedL1Norm(CPUOutputs[Index]);
 		const double CPUGPUCPUAvgL1Norm = GetAveragedL1Norm(CPUGPUCPUOutput);
