Cleaned more API.

_codereview Francisco.Vicente
#rnx

diff --git a/Engine/Plugins/Experimental/NeuralNetworkInference/Source/NeuralNetworkInference/Private/NeuralNetwork.cpp b/Engine/Plugins/Experimental/NeuralNetworkInference/Source/NeuralNetworkInference/Private/NeuralNetwork.cpp
--- a/Engine/Plugins/Experimental/NeuralNetworkInference/Source/NeuralNetworkInference/Private/NeuralNetwork.cpp
+++ b/Engine/Plugins/Experimental/NeuralNetworkInference/Source/NeuralNetworkInference/Private/NeuralNetwork.cpp
@@ -214,7 +214,7 @@ bool UNeuralNetwork::IsGPUSupported() const
 	// UEAndORT
 	if (BackEndForCurrentPlatform == ENeuralBackEnd::UEAndORT)
 	{
-		return FImplBackEndUEAndORT::IsGPUConfigCompatible();
+		return FImplBackEndUEAndORT::IsGPUSupported();
 	}
 
 	// UEOnly
diff --git a/Engine/Plugins/Experimental/NeuralNetworkInference/Source/NeuralNetworkInference/Private/NeuralNetworkImplBackEndUEAndORT.cpp b/Engine/Plugins/Experimental/NeuralNetworkInference/Source/NeuralNetworkInference/Private/NeuralNetworkImplBackEndUEAndORT.cpp
--- a/Engine/Plugins/Experimental/NeuralNetworkInference/Source/NeuralNetworkInference/Private/NeuralNetworkImplBackEndUEAndORT.cpp
+++ b/Engine/Plugins/Experimental/NeuralNetworkInference/Source/NeuralNetworkInference/Private/NeuralNetworkImplBackEndUEAndORT.cpp
@@ -233,11 +233,12 @@ void UNeuralNetwork::FImplBackEndUEAndORT::WarnAndSetDeviceToCPUIfDX12NotEnabled
 {
 	if (InOutDeviceType != ENeuralDeviceType::CPU)
 	{
-		if (!IsGPUConfigCompatible())
+		if (!IsGPUSupported())
 		{
 			InOutDeviceType = ENeuralDeviceType::CPU;
 
 			const FString RHIName = GDynamicRHI->GetName();
+#ifdef PLATFORM_WIN64
 			const FString ErrorMessage = TEXT("On Windows, only DirectX 12 rendering (\"D3D12\") is compatible with the UEAndORT back end of NeuralNetworkInference (NNI). Instead, \"")
 				+ RHIName + TEXT("\" was used. You have the following options:\n\n"
 					"\t1. (Recommended) Switch Unreal Engine to DX12. In order to do that:\n"
@@ -246,6 +247,10 @@ void UNeuralNetwork::FImplBackEndUEAndORT::WarnAndSetDeviceToCPUIfDX12NotEnabled
 					"\t\t - Restart Unreal Engine.\n"
 					"\t2. Alternatively, switch the network to CPU with UNeuralNetwork::SetDeviceType().\n\n"
 					"Network set to CPU provisionally.");
+#else //PLATFORM_WIN64
+			const FString ErrorMessage = TEXT("GPU version is not supported for non-Windows platforms yet. Switch the network to CPU with UNeuralNetwork::SetDeviceType() or run from Windows.\n\n"
+					"Network set to CPU provisionally.");
+#endif //PLATFORM_WIN64
 			UE_LOG(LogNeuralNetworkInference, Warning, TEXT("FImplBackEndUEAndORT::WarnAndSetDeviceToCPUIfDX12NotEnabled(): %s"), *ErrorMessage);
 #if WITH_EDITOR
 			if (bInShouldOpenMessageLog)
@@ -257,7 +262,7 @@ void UNeuralNetwork::FImplBackEndUEAndORT::WarnAndSetDeviceToCPUIfDX12NotEnabled
 	}
 }
 
-bool UNeuralNetwork::FImplBackEndUEAndORT::IsGPUConfigCompatible()
+bool UNeuralNetwork::FImplBackEndUEAndORT::IsGPUSupported()
 {
 #ifdef WITH_UE_AND_ORT_SUPPORT
 #ifdef PLATFORM_WIN64
@@ -267,8 +272,8 @@ bool UNeuralNetwork::FImplBackEndUEAndORT::IsGPUConfigCompatible()
 #endif //PLATFORM_WIN64
 #endif //WITH_UE_AND_ORT_SUPPORT
 
-	// If not Windows and/or if WITH_UE_AND_ORT_SUPPORT not defined, then this should return true because GPU will always work
-	return true;
+	// If not Windows and/or if WITH_UE_AND_ORT_SUPPORT not defined, then this should return false because GPU will not work
+	return false;
 }
 
 bool UNeuralNetwork::FImplBackEndUEAndORT::Load(TSharedPtr<FImplBackEndUEAndORT>& InOutImplBackEndUEAndORT,
@@ -596,7 +601,7 @@ bool UNeuralNetwork::FImplBackEndUEAndORT::ConfigureMembers(const ENeuralDeviceT
 	{
 #ifdef PLATFORM_WIN64
 		// To create a DirectML device we need to check that we're using DX12 first
-		if (!IsGPUConfigCompatible())
+		if (!IsGPUSupported())
 		{
 			UE_LOG(LogNeuralNetworkInference, Warning, TEXT("FImplBackEndUEAndORT::ConfigureMembers(): UEAndORT back end for GPU needs DX12 enabled."));
 			return false;
diff --git a/Engine/Plugins/Experimental/NeuralNetworkInference/Source/NeuralNetworkInference/Private/NeuralNetworkImplBackEndUEAndORT.h b/Engine/Plugins/Experimental/NeuralNetworkInference/Source/NeuralNetworkInference/Private/NeuralNetworkImplBackEndUEAndORT.h
--- a/Engine/Plugins/Experimental/NeuralNetworkInference/Source/NeuralNetworkInference/Private/NeuralNetworkImplBackEndUEAndORT.h
+++ b/Engine/Plugins/Experimental/NeuralNetworkInference/Source/NeuralNetworkInference/Private/NeuralNetworkImplBackEndUEAndORT.h
@@ -99,7 +99,7 @@ public:
 
 	static void WarnAndSetDeviceToCPUIfDX12NotEnabled(ENeuralDeviceType& InOutDeviceType, const bool bInShouldOpenMessageLog);
 
-	static bool IsGPUConfigCompatible();
+	static bool IsGPUSupported();
 
 	static bool Load(TSharedPtr<FImplBackEndUEAndORT>& InOutImplBackEndUEAndORT, FOnAsyncRunCompleted& InOutOnAsyncRunCompletedDelegate,
 		std::atomic<bool>& bInOutIsBackgroundThreadRunning, FCriticalSection& InOutResoucesCriticalSection, TArray<bool>& OutAreInputTensorSizesVariable,
diff --git a/Engine/Plugins/Experimental/NeuralNetworkInference/Source/NeuralNetworkInferenceQA/Private/ModelUnitTester.cpp b/Engine/Plugins/Experimental/NeuralNetworkInference/Source/NeuralNetworkInferenceQA/Private/ModelUnitTester.cpp
--- a/Engine/Plugins/Experimental/NeuralNetworkInference/Source/NeuralNetworkInferenceQA/Private/ModelUnitTester.cpp
+++ b/Engine/Plugins/Experimental/NeuralNetworkInference/Source/NeuralNetworkInferenceQA/Private/ModelUnitTester.cpp
@@ -105,7 +105,12 @@ bool FModelUnitTester::GlobalTest(const FString& InProjectContentDir, const FStr
 	// Speed profiling test - 0 repetitions means that test will not be run
 #ifdef WITH_UE_AND_ORT_SUPPORT
 	const TArray<int32> CPURepetitionsForUEAndORTBackEnd({ 1000, 1000,  50, 1000 });
+#ifdef PLATFORM_WIN64
 	const TArray<int32> GPURepetitionsForUEAndORTBackEnd({ 1000, 1000, 100, 1000 });
+#else //PLATFORM_WIN64
+	UE_LOG(LogNeuralNetworkInferenceQA, Display, TEXT("FModelUnitTester::GlobalTest(): GPU tests disabled for non-Windows platforms."));
+	const TArray<int32> GPURepetitionsForUEAndORTBackEnd({ 0, 0, 0, 0 });
+#endif //PLATFORM_WIN64
 #else //WITH_UE_AND_ORT_SUPPORT
 	const TArray<int32> CPURepetitionsForUEAndORTBackEnd({ 0, 0, 0, 0 });
 	const TArray<int32> GPURepetitionsForUEAndORTBackEnd({ 0, 0, 0, 0 });
@@ -161,7 +166,7 @@ bool FModelUnitTester::ModelLoadAccuracyAndSpeedTests(const FString& InProjectCo
 		{
 			UE_LOG(LogNeuralNetworkInferenceQA, Display, TEXT("OutputTensor[%d] = %s."), TensorIndex, *Network->GetOutputTensor(TensorIndex).GetName());
 		}
-		const bool bShouldRunUEAndORTBackEnd = (InCPURepetitionsForUEAndORTBackEnd[ModelIndex] * InCPURepetitionsForUEAndORTBackEnd[ModelIndex] > 0);
+		const bool bShouldRunUEAndORTBackEnd = (InCPURepetitionsForUEAndORTBackEnd[ModelIndex] + InGPURepetitionsForUEAndORTBackEnd[ModelIndex] > 0);
 		if (bShouldRunUEAndORTBackEnd)
 		{
 			bDidGlobalTestPassed &= ModelAccuracyTest(Network, ENeuralBackEnd::UEAndORT, InInputArrayValues, CPUGroundTruths, GPUGroundTruths);
@@ -211,13 +216,14 @@ bool FModelUnitTester::ModelLoadAccuracyAndSpeedTests(const FString& InProjectCo
 			UE_LOG(LogNeuralNetworkInferenceQA, Warning, TEXT("-------------------- Default UAsset BackEnd should be UEAndORT."));
 			return false;
 		}
-		else if (Network->IsGPUSupported())
+		else if (InGPURepetitionsForUEAndORTBackEnd[ModelIndex] > 0 && Network->IsGPUSupported())
 		{
 			bDidGlobalTestPassed &= ModelSpeedTest(UAssetModelFilePath, ENeuralDeviceType::GPU, ENeuralBackEnd::UEAndORT, InGPURepetitionsForUEAndORTBackEnd[ModelIndex]);
 		}
 		else
 		{
-			UE_LOG(LogNeuralNetworkInferenceQA, Display, TEXT("-------------------- ModelSpeedTest-UEAndORT-GPU skipped (DX12 not enabled)."));
+			UE_LOG(LogNeuralNetworkInferenceQA, Display, TEXT("-------------------- ModelSpeedTest-UEAndORT-GPU skipped (DX12 not enabled or InGPURepetitionsForUEAndORTBackEnd[ModelIndex] = %d is 0)."),
+				InGPURepetitionsForUEAndORTBackEnd[ModelIndex]);
 		}
 #endif //WITH_UE_AND_ORT_SUPPORT
 
