Async demo up (but not working& disabled)
#rnx

diff --git a/Engine/Plugins/Experimental/NeuralNetworkInference/Source/NeuralNetworkInference/Private/NeuralNetworkImplBackEndUEAndORT.cpp b/Engine/Plugins/Experimental/NeuralNetworkInference/Source/NeuralNetworkInference/Private/NeuralNetworkImplBackEndUEAndORT.cpp
--- a/Engine/Plugins/Experimental/NeuralNetworkInference/Source/NeuralNetworkInference/Private/NeuralNetworkImplBackEndUEAndORT.cpp
+++ b/Engine/Plugins/Experimental/NeuralNetworkInference/Source/NeuralNetworkInference/Private/NeuralNetworkImplBackEndUEAndORT.cpp
@@ -370,7 +370,7 @@ void UNeuralNetwork::FImplBackEndUEAndORT::IsAsyncTaskDone() const
 	if (NeuralNetworkAsyncTask && !NeuralNetworkAsyncTask->IsDone())
 	{
 		UE_LOG(LogNeuralNetworkInference, Warning,
-			TEXT("FImplBackEndUEAndORT::Run(): Previous async run had not been completed. Blocking thread until it is completed."));
+			TEXT("FImplBackEndUEAndORT::IsAsyncTaskDone(): Previous async run had not been completed. Blocking thread until it is completed."));
 		NeuralNetworkAsyncTask->EnsureCompletion(/*bDoWorkOnThisThreadIfNotStarted*/true);
 	}
 }
