Improved doc

_codereview Francisco.Vicente
#rnx

diff --git a/Engine/Plugins/Experimental/NeuralNetworkInference/Source/NeuralNetworkInference/Private/NeuralNetworkImplBackEndUEAndORT.cpp b/Engine/Plugins/Experimental/NeuralNetworkInference/Source/NeuralNetworkInference/Private/NeuralNetworkImplBackEndUEAndORT.cpp
--- a/Engine/Plugins/Experimental/NeuralNetworkInference/Source/NeuralNetworkInference/Private/NeuralNetworkImplBackEndUEAndORT.cpp
+++ b/Engine/Plugins/Experimental/NeuralNetworkInference/Source/NeuralNetworkInference/Private/NeuralNetworkImplBackEndUEAndORT.cpp
@@ -532,17 +532,17 @@ void UNeuralNetwork::FImplBackEndUEAndORT::RunSessionImpl(const ENeuralDeviceTyp
 			// TODO: Remove this sync point and move session run to render thread
 			FNeuralNetworkInferenceUtils::WaitUntilRHIFinished();
 		}
-	}
 
-	if (InDeviceType == ENeuralDeviceType::GPU)
-	{
+		// Profiling init
 		FNNIGPUProfiler::Instance()->EventBegin("NNI:SessionRun");
 	}
 
+	// Actual ORT network run
 	Session->Run(Ort::RunOptions{ nullptr },
 		InputTensorNames.GetData(), &InputOrtTensors[0], InputTensorNames.Num(),
 		OutputTensorNames.GetData(), &OutputOrtTensors[0], OutputTensorNames.Num());
 
+	// Profiling end
 	if (InDeviceType == ENeuralDeviceType::GPU)
 	{
 		FNNIGPUProfiler::Instance()->EventEnd();
diff --git a/Engine/Plugins/Experimental/NeuralNetworkInference/Source/NeuralNetworkInference/Public/NeuralNetwork.h b/Engine/Plugins/Experimental/NeuralNetworkInference/Source/NeuralNetworkInference/Public/NeuralNetwork.h
--- a/Engine/Plugins/Experimental/NeuralNetworkInference/Source/NeuralNetworkInference/Public/NeuralNetwork.h
+++ b/Engine/Plugins/Experimental/NeuralNetworkInference/Source/NeuralNetworkInference/Public/NeuralNetwork.h
@@ -9,13 +9,21 @@
 #include "NeuralNetwork.generated.h"
 
 /**
- * Whether UNeuralNetwork::Run() will block the thread until completed (Synchronous), or whether it will run on a background thread, not blocking the calling thread (Asynchronous).
+ * Whether UNeuralNetwork::Run() will block the thread until completed (Synchronous), or whether it will run on a background thread,
+ * not blocking the calling thread (Asynchronous).
  */
 UENUM()
 enum class ENeuralNetworkSynchronousMode : uint8
 {
 	Synchronous, /* UNeuralNetwork::Run() will block the thread until the network evaluation (i.e., forward pass) has finished. */
-	Asynchronous /* UNeuralNetwork::Run() will initialize a forward pass request on a background thread, not blocking the thread that called it. The user should register to UNeuralNetwork's delegate to know when the forward pass has finished. */
+	/**
+	 * UNeuralNetwork::Run() will initialize a forward pass request on a background thread, not blocking the thread that called it.
+	 * The user should register to UNeuralNetwork's delegate to know when the forward pass has finished.
+	 *
+	 * Very important: It takes ~1 millisecond to start the background thread. If your network runs synchronously faster than 1 msec,
+	 * using asynchronous running will make the game (main) thread slower than running it synchronously.
+	 */
+	Asynchronous
 };
 
 /**
