Fixed bug in Async NNI.

#rnx

diff --git a/Engine/Plugins/Experimental/NeuralNetworkInference/Source/NeuralNetworkInference/Private/NeuralNetworkImplBackEndUEAndORT.cpp b/Engine/Plugins/Experimental/NeuralNetworkInference/Source/NeuralNetworkInference/Private/NeuralNetworkImplBackEndUEAndORT.cpp
--- a/Engine/Plugins/Experimental/NeuralNetworkInference/Source/NeuralNetworkInference/Private/NeuralNetworkImplBackEndUEAndORT.cpp
+++ b/Engine/Plugins/Experimental/NeuralNetworkInference/Source/NeuralNetworkInference/Private/NeuralNetworkImplBackEndUEAndORT.cpp
@@ -229,6 +229,11 @@ IDMLDevice* FPrivateImplBackEndUEAndORT::FDMLDeviceList::Add(ID3D12Device* Devic
 /* UNeuralNetwork public functions
  *****************************************************************************/
 
+UNeuralNetwork::FImplBackEndUEAndORT::~FImplBackEndUEAndORT()
+{
+	EnsureAsyncTaskCompletion(/*bShouldWarnIfNotDone*/false);
+}
+
 void UNeuralNetwork::FImplBackEndUEAndORT::WarnAndSetDeviceToCPUIfDX12NotEnabled(ENeuralDeviceType& InOutDeviceType, const bool bInShouldOpenMessageLog)
 {
 	if (InOutDeviceType != ENeuralDeviceType::CPU)
@@ -291,7 +296,7 @@ bool UNeuralNetwork::FImplBackEndUEAndORT::Load(TSharedPtr<FImplBackEndUEAndORT>
 
 		if (InOutImplBackEndUEAndORT.IsValid())
 		{
-			InOutImplBackEndUEAndORT->IsAsyncTaskDone();
+			InOutImplBackEndUEAndORT->EnsureAsyncTaskCompletion(/*bShouldWarnIfNotDone*/true);
 		}
 
 		// Initialize and configure InOutImplBackEndUEAndORT
@@ -365,12 +370,15 @@ bool UNeuralNetwork::FImplBackEndUEAndORT::Load(TSharedPtr<FImplBackEndUEAndORT>
 }
 
 #ifdef WITH_UE_AND_ORT_SUPPORT
-void UNeuralNetwork::FImplBackEndUEAndORT::IsAsyncTaskDone() const
+void UNeuralNetwork::FImplBackEndUEAndORT::EnsureAsyncTaskCompletion(const bool bShouldWarnIfNotDone) const
 {
 	if (NeuralNetworkAsyncTask && !NeuralNetworkAsyncTask->IsDone())
 	{
-		UE_LOG(LogNeuralNetworkInference, Warning,
-			TEXT("FImplBackEndUEAndORT::IsAsyncTaskDone(): Previous async run had not been completed. Blocking thread until it is completed."));
+		if (bShouldWarnIfNotDone)
+		{
+			UE_LOG(LogNeuralNetworkInference, Warning,
+				TEXT("FImplBackEndUEAndORT::EnsureAsyncTaskCompletion(): Previous async run had not been completed. Blocking thread until it is completed."));
+		}
 		NeuralNetworkAsyncTask->EnsureCompletion(/*bDoWorkOnThisThreadIfNotStarted*/true);
 	}
 }
@@ -412,7 +420,7 @@ void UNeuralNetwork::FImplBackEndUEAndORT::Run(const ENeuralNetworkSynchronousMo
 	{
 		const FRedirectCoutAndCerrToUeLog RedirectCoutAndCerrToUeLog;
 
-		IsAsyncTaskDone();
+		EnsureAsyncTaskCompletion(/*bShouldWarnIfNotDone*/true);
 		NeuralNetworkAsyncTask->GetTask().SetRunSessionArgs(InSynchronousMode, InDeviceType, InInputDeviceType, InOutputDeviceType);
 
 		// Run UNeuralNetwork
diff --git a/Engine/Plugins/Experimental/NeuralNetworkInference/Source/NeuralNetworkInference/Private/NeuralNetworkImplBackEndUEAndORT.h b/Engine/Plugins/Experimental/NeuralNetworkInference/Source/NeuralNetworkInference/Private/NeuralNetworkImplBackEndUEAndORT.h
--- a/Engine/Plugins/Experimental/NeuralNetworkInference/Source/NeuralNetworkInference/Private/NeuralNetworkImplBackEndUEAndORT.h
+++ b/Engine/Plugins/Experimental/NeuralNetworkInference/Source/NeuralNetworkInference/Private/NeuralNetworkImplBackEndUEAndORT.h
@@ -97,6 +97,8 @@ public:
 	TArray<FNeuralTensor> InputTensors;
 	TArray<FNeuralTensor> OutputTensors;
 
+	~FImplBackEndUEAndORT();
+
 	static void WarnAndSetDeviceToCPUIfDX12NotEnabled(ENeuralDeviceType& InOutDeviceType, const bool bInShouldOpenMessageLog);
 
 	static bool IsGPUSupported();
@@ -139,7 +141,7 @@ private:
 
 	TUniquePtr<FAsyncTask<FNeuralNetworkAsyncTask>> NeuralNetworkAsyncTask;
 
-	void IsAsyncTaskDone() const;
+	void EnsureAsyncTaskCompletion(const bool bShouldWarnIfNotDone) const;
 
 	static bool InitializedAndConfigureMembers(TSharedPtr<FImplBackEndUEAndORT>& InOutImplBackEndUEAndORT,
 		FOnAsyncRunCompleted& InOutOnAsyncRunCompletedDelegate, std::atomic<bool>& bInOutIsBackgroundThreadRunning,
