Header cleanup + CIS fix

_codereview Francisco.Vicente
#rnx

diff --git a/Engine/Plugins/Experimental/NeuralNetworkInference/Source/NeuralNetworkInference/Private/NeuralNetworkImplBackEndUEAndORT.cpp b/Engine/Plugins/Experimental/NeuralNetworkInference/Source/NeuralNetworkInference/Private/NeuralNetworkImplBackEndUEAndORT.cpp
--- a/Engine/Plugins/Experimental/NeuralNetworkInference/Source/NeuralNetworkInference/Private/NeuralNetworkImplBackEndUEAndORT.cpp
+++ b/Engine/Plugins/Experimental/NeuralNetworkInference/Source/NeuralNetworkInference/Private/NeuralNetworkImplBackEndUEAndORT.cpp
@@ -226,12 +226,47 @@ IDMLDevice* FPrivateImplBackEndUEAndORT::FDMLDeviceList::Add(ID3D12Device* Devic
 
 
 
+/* FNeuralNetworkAsyncTask public functions
+ *****************************************************************************/
+
+#ifdef WITH_UE_AND_ORT_SUPPORT
+UNeuralNetwork::FImplBackEndUEAndORT::FNeuralNetworkAsyncTask::FNeuralNetworkAsyncTask(UNeuralNetwork::FImplBackEndUEAndORT* InBackEnd)
+	: BackEnd(InBackEnd)
+{}
+
+void UNeuralNetwork::FImplBackEndUEAndORT::FNeuralNetworkAsyncTask::SetRunSessionArgs(const ENeuralNetworkSynchronousMode InSyncMode, const ENeuralDeviceType InDeviceType,
+	const ENeuralDeviceType InInputDeviceType, const ENeuralDeviceType InOutputDeviceType)
+{
+	const FScopeLock ResourcesLock(&BackEnd->ResoucesCriticalSection);
+	SyncMode = InSyncMode;
+	DeviceType = InDeviceType;
+	InputDeviceType = InInputDeviceType;
+	OutputDeviceType = InOutputDeviceType;
+}
+
+void UNeuralNetwork::FImplBackEndUEAndORT::FNeuralNetworkAsyncTask::DoWork()
+{
+	if (SyncMode == ENeuralNetworkSynchronousMode::Asynchronous)
+	{
+		BackEnd->RunSessionAsync(DeviceType, InputDeviceType, OutputDeviceType);
+	}
+	else
+	{
+		BackEnd->RunSessionSync(DeviceType, InputDeviceType, OutputDeviceType);
+	}
+}
+#endif //WITH_UE_AND_ORT_SUPPORT
+
+
+
 /* UNeuralNetwork public functions
  *****************************************************************************/
 
 UNeuralNetwork::FImplBackEndUEAndORT::~FImplBackEndUEAndORT()
 {
+#ifdef WITH_UE_AND_ORT_SUPPORT
 	EnsureAsyncTaskCompletion(/*bShouldWarnIfNotDone*/false);
+#endif //WITH_UE_AND_ORT_SUPPORT
 }
 
 void UNeuralNetwork::FImplBackEndUEAndORT::WarnAndSetDeviceToCPUIfDX12NotEnabled(ENeuralDeviceType& InOutDeviceType, const bool bInShouldOpenMessageLog)
diff --git a/Engine/Plugins/Experimental/NeuralNetworkInference/Source/NeuralNetworkInference/Private/NeuralNetworkImplBackEndUEAndORT.h b/Engine/Plugins/Experimental/NeuralNetworkInference/Source/NeuralNetworkInference/Private/NeuralNetworkImplBackEndUEAndORT.h
--- a/Engine/Plugins/Experimental/NeuralNetworkInference/Source/NeuralNetworkInference/Private/NeuralNetworkImplBackEndUEAndORT.h
+++ b/Engine/Plugins/Experimental/NeuralNetworkInference/Source/NeuralNetworkInference/Private/NeuralNetworkImplBackEndUEAndORT.h
@@ -36,60 +36,6 @@ NNI_THIRD_PARTY_INCLUDES_END
 
 struct UNeuralNetwork::FImplBackEndUEAndORT
 {
-	/// Helper class to run session as an async task
-	class FNeuralNetworkAsyncTask : public FNonAbandonableTask
-	{
-		friend class FAsyncTask<FNeuralNetworkAsyncTask>;
-
-	public:
-		FNeuralNetworkAsyncTask(UNeuralNetwork::FImplBackEndUEAndORT* InBackEnd)
-			: BackEnd(InBackEnd)
-		{}
-
-		void SetRunSessionArgs(const ENeuralNetworkSynchronousMode InSyncMode, const ENeuralDeviceType InDeviceType,
-			const ENeuralDeviceType InInputDeviceType, const ENeuralDeviceType InOutputDeviceType)
-		{
-#ifdef WITH_UE_AND_ORT_SUPPORT
-			const FScopeLock ResourcesLock(&BackEnd->ResoucesCriticalSection);
-
-			SyncMode = InSyncMode;
-			DeviceType = InDeviceType;
-			InputDeviceType = InInputDeviceType;
-			OutputDeviceType = InOutputDeviceType;
-#endif
-		}
-
-	protected:
-		void DoWork()
-		{
-#ifdef WITH_UE_AND_ORT_SUPPORT
-			if (SyncMode == ENeuralNetworkSynchronousMode::Asynchronous)
-			{
-				BackEnd->RunSessionAsync(DeviceType, InputDeviceType, OutputDeviceType);
-			}
-			else
-			{
-				BackEnd->RunSessionSync(DeviceType, InputDeviceType, OutputDeviceType);
-			}
-#endif
-		}
-
-		// This next section of code needs to be here. Not important as to why.
-		FORCEINLINE TStatId GetStatId() const
-		{
-			RETURN_QUICK_DECLARE_CYCLE_STAT(FNeuralNetworkAsyncTask, STATGROUP_ThreadPoolAsyncTasks);
-		}
-
-	private:
-		UNeuralNetwork::FImplBackEndUEAndORT* BackEnd;
-
-		/// Variables that could change on each inference run
-		ENeuralNetworkSynchronousMode SyncMode;
-		ENeuralDeviceType DeviceType;
-		ENeuralDeviceType InputDeviceType;
-		ENeuralDeviceType OutputDeviceType;
-	};
-
 public:
 	/**
 	 * InputTensors and OutputTensors represent the input and output TArray<FNeuralTensor> of the network, respectively.
@@ -138,7 +84,30 @@ private:
 	TArray<const char*> InputTensorNames; /* Tensor names */
 	TArray<Ort::Value> OutputOrtTensors; /* Actual ONNXRuntime tensors */
 	TArray<const char*> OutputTensorNames; /* Tensor names */
-
+	
+	// Helper class to run session as an async task
+	class FNeuralNetworkAsyncTask : public FNonAbandonableTask
+	{
+		friend class FAsyncTask<FNeuralNetworkAsyncTask>;
+	public:
+		FNeuralNetworkAsyncTask(UNeuralNetwork::FImplBackEndUEAndORT* InBackEnd);
+		void SetRunSessionArgs(const ENeuralNetworkSynchronousMode InSyncMode, const ENeuralDeviceType InDeviceType,
+			const ENeuralDeviceType InInputDeviceType, const ENeuralDeviceType InOutputDeviceType);
+	protected:
+		void DoWork();
+		// This next section of code needs to be here. Not important as to why.
+		FORCEINLINE TStatId GetStatId() const
+		{
+			RETURN_QUICK_DECLARE_CYCLE_STAT(FNeuralNetworkAsyncTask, STATGROUP_ThreadPoolAsyncTasks);
+		}
+	private:
+		UNeuralNetwork::FImplBackEndUEAndORT* BackEnd;
+		// Variables that could change on each inference run
+		ENeuralNetworkSynchronousMode SyncMode;
+		ENeuralDeviceType DeviceType;
+		ENeuralDeviceType InputDeviceType;
+		ENeuralDeviceType OutputDeviceType;
+	};
 	TUniquePtr<FAsyncTask<FNeuralNetworkAsyncTask>> NeuralNetworkAsyncTask;
 
 	void EnsureAsyncTaskCompletion(const bool bShouldWarnIfNotDone) const;
