FImplBackEndUEAndORT cleanup

_codereview Francisco.Vicente
#rnx

diff --git a/Engine/Plugins/Experimental/NeuralNetworkInference/Source/NeuralNetworkInference/Private/NeuralNetworkImplBackEndUEAndORT.cpp b/Engine/Plugins/Experimental/NeuralNetworkInference/Source/NeuralNetworkInference/Private/NeuralNetworkImplBackEndUEAndORT.cpp
--- a/Engine/Plugins/Experimental/NeuralNetworkInference/Source/NeuralNetworkInference/Private/NeuralNetworkImplBackEndUEAndORT.cpp
+++ b/Engine/Plugins/Experimental/NeuralNetworkInference/Source/NeuralNetworkInference/Private/NeuralNetworkImplBackEndUEAndORT.cpp
@@ -133,6 +133,7 @@
 };
 
 
+
 #ifdef PLATFORM_WIN64
 
 /* FPrivateImplBackEndUEAndORT auxiliary class
@@ -259,9 +260,19 @@
 
 
 
-/* UNeuralNetwork public functions
+/* FImplBackEndUEAndORT 'structors
  *****************************************************************************/
 
+UNeuralNetwork::FImplBackEndUEAndORT::FImplBackEndUEAndORT(FOnAsyncRunCompleted& InOutOnAsyncRunCompletedDelegate,
+	std::atomic<bool>& bInIsBackgroundThreadRunning, FCriticalSection& InResoucesCriticalSection)
+#ifdef WITH_UE_AND_ORT_SUPPORT
+	: OnAsyncRunCompletedDelegate(InOutOnAsyncRunCompletedDelegate)
+	, bIsBackgroundThreadRunning(bInIsBackgroundThreadRunning)
+	, ResoucesCriticalSection(InResoucesCriticalSection)
+#endif
+{
+}
+
 UNeuralNetwork::FImplBackEndUEAndORT::~FImplBackEndUEAndORT()
 {
 #ifdef WITH_UE_AND_ORT_SUPPORT
@@ -269,6 +280,11 @@
 #endif //WITH_UE_AND_ORT_SUPPORT
 }
 
+
+
+/* FImplBackEndUEAndORT public functions
+ *****************************************************************************/
+
 void UNeuralNetwork::FImplBackEndUEAndORT::WarnAndSetDeviceToCPUIfDX12NotEnabled(ENeuralDeviceType& InOutDeviceType, const bool bInShouldOpenMessageLog)
 {
 	if (InOutDeviceType != ENeuralDeviceType::CPU)
@@ -404,47 +420,6 @@
 #endif //WITH_UE_AND_ORT_SUPPORT
 }
 
-#ifdef WITH_UE_AND_ORT_SUPPORT
-void UNeuralNetwork::FImplBackEndUEAndORT::EnsureAsyncTaskCompletion(const bool bShouldWarnIfNotDone) const
-{
-	if (NeuralNetworkAsyncTask && !NeuralNetworkAsyncTask->IsDone())
-	{
-		if (bShouldWarnIfNotDone)
-		{
-			UE_LOG(LogNeuralNetworkInference, Warning,
-				TEXT("FImplBackEndUEAndORT::EnsureAsyncTaskCompletion(): Previous async run had not been completed. Blocking thread until it is completed."));
-		}
-		NeuralNetworkAsyncTask->EnsureCompletion(/*bDoWorkOnThisThreadIfNotStarted*/true);
-	}
-}
-
-void UNeuralNetwork::FImplBackEndUEAndORT::ClearResources()
-{
-#ifdef PLATFORM_WIN64
-	if (DmlApi)
-	{
-		const int32 Num = DmlGPUResources.Num();
-		for (int32 Index = 0; Index < Num; ++Index)
-		{
-			DmlApi->FreeGPUAllocation(DmlGPUResources[Index]);
-		}
-
-		DmlGPUResources.Reset(0);
-		DmlApi = nullptr;
-	}
-#endif //PLATFORM_WIN64
-}
-
-UNeuralNetwork::FImplBackEndUEAndORT::FImplBackEndUEAndORT(FOnAsyncRunCompleted& InOutOnAsyncRunCompletedDelegate,
-	std::atomic<bool>& bInIsBackgroundThreadRunning, FCriticalSection& InResoucesCriticalSection)
-	: OnAsyncRunCompletedDelegate(InOutOnAsyncRunCompletedDelegate)
-	, bIsBackgroundThreadRunning(bInIsBackgroundThreadRunning)
-	, ResoucesCriticalSection(InResoucesCriticalSection)
-{
-}
-
-#endif //WITH_UE_AND_ORT_SUPPORT
-
 void UNeuralNetwork::FImplBackEndUEAndORT::Run(const ENeuralNetworkSynchronousMode InSynchronousMode, const ENeuralDeviceType InDeviceType,
 	const ENeuralDeviceType InInputDeviceType, const ENeuralDeviceType InOutputDeviceType)
 {
@@ -486,119 +461,28 @@
 #endif //WITH_UE_AND_ORT_SUPPORT
 }
 
-/* UNeuralNetwork private functions
+
+
+/* FImplBackEndUEAndORT private functions
  *****************************************************************************/
 
 #ifdef WITH_UE_AND_ORT_SUPPORT
-
-void UNeuralNetwork::FImplBackEndUEAndORT::RunSessionAsync(const ENeuralDeviceType InDeviceType, const ENeuralDeviceType InInputDeviceType,
-	const ENeuralDeviceType InOutputDeviceType)
+void UNeuralNetwork::FImplBackEndUEAndORT::EnsureAsyncTaskCompletion(const bool bShouldWarnIfNotDone) const
 {
-	const FScopeLock ResourcesLock(&ResoucesCriticalSection);
-
-	RunSessionImpl(InDeviceType, InInputDeviceType, InOutputDeviceType);
-
-	OnAsyncRunCompletedDelegate.ExecuteIfBound();
-	bIsBackgroundThreadRunning = false;
-}
-
-void UNeuralNetwork::FImplBackEndUEAndORT::RunSessionSync(const ENeuralDeviceType InDeviceType, const ENeuralDeviceType InInputDeviceType,
-	const ENeuralDeviceType InOutputDeviceType)
-{
-	RunSessionImpl(InDeviceType, InInputDeviceType, InOutputDeviceType);
-}
-
-// Used when uploading tensors to GPU
-// NOTE: Upload parameter is not yet used, we plan to use it in the future
-BEGIN_SHADER_PARAMETER_STRUCT(FUploadTensorParameters, )
-	RDG_BUFFER_ACCESS(Upload, ERHIAccess::CopySrc)
-	RDG_BUFFER_ACCESS(Input, ERHIAccess::CopyDest)
-END_SHADER_PARAMETER_STRUCT()
-
-void UNeuralNetwork::FImplBackEndUEAndORT::RunSessionImpl(const ENeuralDeviceType InDeviceType, const ENeuralDeviceType InInputDeviceType,
-	const ENeuralDeviceType InOutputDeviceType)
-{
-	if (InDeviceType == ENeuralDeviceType::GPU)
+	if (NeuralNetworkAsyncTask && !NeuralNetworkAsyncTask->IsDone())
 	{
-		// Copy data to GPU (if required)
-		bool bNeedsGPUCopy = false;
-
-		for (auto& InputTensor : InputTensors)
+		if (bShouldWarnIfNotDone)
 		{
-			if (InputTensor.GetTensorTypeGPU() != ENeuralTensorTypeGPU::Input)
-				continue;
-
-			bNeedsGPUCopy = true;
-
-			ENQUEUE_RENDER_COMMAND(UploadTensorToGPU)([this, InputTensor] (FRHICommandListImmediate& RHICmdList)
-			{
-				FRDGBuilder GraphBuilder(RHICmdList, RDG_EVENT_NAME("UploadTensorToGPU"));
-
-				// Set parameters
-				TRefCountPtr< FRDGPooledBuffer >&	PooledBuffer = InputTensor.GetPooledBuffer();
-				FRDGBufferRef						InputBufferRef = GraphBuilder.RegisterExternalBuffer(PooledBuffer);
-
-				FUploadTensorParameters*			UploadParameters = GraphBuilder.AllocParameters<FUploadTensorParameters>();
-
-				UploadParameters->Input = InputBufferRef;
-
-				GraphBuilder.AddPass(
-					RDG_EVENT_NAME("NNI:UploadTensor:%s", InputTensor.GetNameData()),
-					FUploadTensorParameters::FTypeInfo::GetStructMetadata(),
-					UploadParameters,
-					ERDGPassFlags::Copy | ERDGPassFlags::NeverCull,
-					[this, UploadParameters](FRHICommandListImmediate& RHICmdList)
-					{
-						FRHIBuffer* InputBuffer = UploadParameters->Input->GetRHI();
-
-						// NOTE: We're using UAVMask to trigger the UAV barrier in RDG
-						RHICmdList.Transition(FRHITransitionInfo(InputBuffer, ERHIAccess::CopyDest, ERHIAccess::UAVMask));
-					}
-				);
-
-				GraphBuilder.Execute();
-			});
+			UE_LOG(LogNeuralNetworkInference, Warning,
+				TEXT("FImplBackEndUEAndORT::EnsureAsyncTaskCompletion(): Previous async run had not been completed. Blocking thread until it is completed."));
 		}
-
-		if (bNeedsGPUCopy)
-		{
-			ENQUEUE_RENDER_COMMAND(FlushUploadTensorToGPU)([this] (FRHICommandListImmediate& RHICmdList)
-			{
-				FRDGBuilder GraphBuilder(RHICmdList, RDG_EVENT_NAME("NNI:FlushUploadTensorsToGPU"));
-
-				RHICmdList.SubmitCommandsHint();
-				RHICmdList.ImmediateFlush(EImmediateFlushType::FlushRHIThread);
-
-				GraphBuilder.Execute();
-			});
-
-			// TODO: Remove this sync point and move session run to render thread
-			FNeuralNetworkInferenceUtils::WaitUntilRHIFinished();
-		}
-
-		// Profiling init
-		FNNIGPUProfiler::Instance()->EventBegin("NNI:SessionRun");
-	}
-
-	// Actual ORT network run
-	Session->Run(Ort::RunOptions{ nullptr },
-		InputTensorNames.GetData(), &InputOrtTensors[0], InputTensorNames.Num(),
-		OutputTensorNames.GetData(), &OutputOrtTensors[0], OutputTensorNames.Num());
-
-	// Profiling end
-	if (InDeviceType == ENeuralDeviceType::GPU)
-	{
-		FNNIGPUProfiler::Instance()->EventEnd();
+		NeuralNetworkAsyncTask->EnsureCompletion(/*bDoWorkOnThisThreadIfNotStarted*/true);
 	}
 }
 
-bool UNeuralNetwork::FImplBackEndUEAndORT::InitializedAndConfigureMembers(
-			TSharedPtr<FImplBackEndUEAndORT>& InOutImplBackEndUEAndORT,
-			FOnAsyncRunCompleted& InOutOnAsyncRunCompletedDelegate,
-			std::atomic<bool>& bInOutIsBackgroundThreadRunning,
-			FCriticalSection& InOutResoucesCriticalSection,
-			const FString& InModelFullFilePath,
-			const ENeuralDeviceType InDeviceType)
+bool UNeuralNetwork::FImplBackEndUEAndORT::InitializedAndConfigureMembers(TSharedPtr<FImplBackEndUEAndORT>& InOutImplBackEndUEAndORT,
+	FOnAsyncRunCompleted& InOutOnAsyncRunCompletedDelegate, std::atomic<bool>& bInOutIsBackgroundThreadRunning, FCriticalSection& InOutResoucesCriticalSection,
+	const FString& InModelFullFilePath, const ENeuralDeviceType InDeviceType)
 {
 	// Initialize InOutImplBackEndUEAndORT
 	if (!InOutImplBackEndUEAndORT.IsValid())
@@ -989,7 +873,6 @@
 }
 
 #ifdef PLATFORM_WIN64
-
 bool UNeuralNetwork::FImplBackEndUEAndORT::LinkTensorResourceToONNXRuntime(FNeuralTensor& InOutTensor, Ort::Value& InOutOrtTensor, void* D3DResource)
 {
 	if (!DmlApi)
@@ -1044,7 +927,124 @@
 
 	return true;
 }
-
 #endif
 
+void UNeuralNetwork::FImplBackEndUEAndORT::RunSessionAsync(const ENeuralDeviceType InDeviceType, const ENeuralDeviceType InInputDeviceType,
+	const ENeuralDeviceType InOutputDeviceType)
+{
+	const FScopeLock ResourcesLock(&ResoucesCriticalSection);
+
+	RunSessionImpl(InDeviceType, InInputDeviceType, InOutputDeviceType);
+
+	OnAsyncRunCompletedDelegate.ExecuteIfBound();
+	bIsBackgroundThreadRunning = false;
+}
+
+void UNeuralNetwork::FImplBackEndUEAndORT::RunSessionSync(const ENeuralDeviceType InDeviceType, const ENeuralDeviceType InInputDeviceType,
+	const ENeuralDeviceType InOutputDeviceType)
+{
+	RunSessionImpl(InDeviceType, InInputDeviceType, InOutputDeviceType);
+}
+
+// Used when uploading tensors to GPU
+// NOTE: Upload parameter is not yet used, we plan to use it in the future
+BEGIN_SHADER_PARAMETER_STRUCT(FUploadTensorParameters, )
+RDG_BUFFER_ACCESS(Upload, ERHIAccess::CopySrc)
+RDG_BUFFER_ACCESS(Input, ERHIAccess::CopyDest)
+END_SHADER_PARAMETER_STRUCT()
+
+void UNeuralNetwork::FImplBackEndUEAndORT::RunSessionImpl(const ENeuralDeviceType InDeviceType, const ENeuralDeviceType InInputDeviceType,
+	const ENeuralDeviceType InOutputDeviceType)
+{
+	if (InDeviceType == ENeuralDeviceType::GPU)
+	{
+		// Copy data to GPU (if required)
+		bool bNeedsGPUCopy = false;
+
+		for (auto& InputTensor : InputTensors)
+		{
+			if (InputTensor.GetTensorTypeGPU() != ENeuralTensorTypeGPU::Input)
+				continue;
+
+			bNeedsGPUCopy = true;
+
+			ENQUEUE_RENDER_COMMAND(UploadTensorToGPU)([this, InputTensor](FRHICommandListImmediate& RHICmdList)
+				{
+					FRDGBuilder GraphBuilder(RHICmdList, RDG_EVENT_NAME("UploadTensorToGPU"));
+
+					// Set parameters
+					TRefCountPtr< FRDGPooledBuffer >& PooledBuffer = InputTensor.GetPooledBuffer();
+					FRDGBufferRef						InputBufferRef = GraphBuilder.RegisterExternalBuffer(PooledBuffer);
+
+					FUploadTensorParameters* UploadParameters = GraphBuilder.AllocParameters<FUploadTensorParameters>();
+
+					UploadParameters->Input = InputBufferRef;
+
+					GraphBuilder.AddPass(
+						RDG_EVENT_NAME("NNI:UploadTensor:%s", InputTensor.GetNameData()),
+						FUploadTensorParameters::FTypeInfo::GetStructMetadata(),
+						UploadParameters,
+						ERDGPassFlags::Copy | ERDGPassFlags::NeverCull,
+						[this, UploadParameters](FRHICommandListImmediate& RHICmdList)
+						{
+							FRHIBuffer* InputBuffer = UploadParameters->Input->GetRHI();
+
+							// NOTE: We're using UAVMask to trigger the UAV barrier in RDG
+							RHICmdList.Transition(FRHITransitionInfo(InputBuffer, ERHIAccess::CopyDest, ERHIAccess::UAVMask));
+						}
+					);
+
+					GraphBuilder.Execute();
+				});
+		}
+
+		if (bNeedsGPUCopy)
+		{
+			ENQUEUE_RENDER_COMMAND(FlushUploadTensorToGPU)([this](FRHICommandListImmediate& RHICmdList)
+				{
+					FRDGBuilder GraphBuilder(RHICmdList, RDG_EVENT_NAME("NNI:FlushUploadTensorsToGPU"));
+
+					RHICmdList.SubmitCommandsHint();
+					RHICmdList.ImmediateFlush(EImmediateFlushType::FlushRHIThread);
+
+					GraphBuilder.Execute();
+				});
+
+			// TODO: Remove this sync point and move session run to render thread
+			FNeuralNetworkInferenceUtils::WaitUntilRHIFinished();
+		}
+
+		// Profiling init
+		FNNIGPUProfiler::Instance()->EventBegin("NNI:SessionRun");
+	}
+
+	// Actual ORT network run
+	Session->Run(Ort::RunOptions{ nullptr },
+		InputTensorNames.GetData(), &InputOrtTensors[0], InputTensorNames.Num(),
+		OutputTensorNames.GetData(), &OutputOrtTensors[0], OutputTensorNames.Num());
+
+	// Profiling end
+	if (InDeviceType == ENeuralDeviceType::GPU)
+	{
+		FNNIGPUProfiler::Instance()->EventEnd();
+	}
+}
+
+void UNeuralNetwork::FImplBackEndUEAndORT::ClearResources()
+{
+#ifdef PLATFORM_WIN64
+	if (DmlApi)
+	{
+		const int32 Num = DmlGPUResources.Num();
+		for (int32 Index = 0; Index < Num; ++Index)
+		{
+			DmlApi->FreeGPUAllocation(DmlGPUResources[Index]);
+		}
+
+		DmlGPUResources.Reset(0);
+		DmlApi = nullptr;
+	}
+#endif //PLATFORM_WIN64
+}
+
 #endif //WITH_UE_AND_ORT_SUPPORT
diff --git a/Engine/Plugins/Experimental/NeuralNetworkInference/Source/NeuralNetworkInference/Private/NeuralNetworkImplBackEndUEAndORT.h b/Engine/Plugins/Experimental/NeuralNetworkInference/Source/NeuralNetworkInference/Private/NeuralNetworkImplBackEndUEAndORT.h
--- a/Engine/Plugins/Experimental/NeuralNetworkInference/Source/NeuralNetworkInference/Private/NeuralNetworkImplBackEndUEAndORT.h
+++ b/Engine/Plugins/Experimental/NeuralNetworkInference/Source/NeuralNetworkInference/Private/NeuralNetworkImplBackEndUEAndORT.h
@@ -11,7 +11,7 @@
 #include "DynamicRHI.h"
 
 #if defined(WITH_UE_AND_ORT_SUPPORT) && defined(PLATFORM_WIN64)
-	// Disable NOMINMAX & WIN32_LEAN_AND_MEAN defines to avoid compiler warnings
+// Disable NOMINMAX & WIN32_LEAN_AND_MEAN defines to avoid compiler warnings
 #pragma push_macro("NOMINMAX")
 #pragma push_macro("WIN32_LEAN_AND_MEAN")
 #undef NOMINMAX
@@ -43,6 +43,9 @@ public:
 	TArray<FNeuralTensor> InputTensors;
 	TArray<FNeuralTensor> OutputTensors;
 
+	FImplBackEndUEAndORT(FOnAsyncRunCompleted& InOutOnAsyncRunCompletedDelegate, std::atomic<bool>& bInIsBackgroundThreadRunning,
+		FCriticalSection& InResoucesCriticalSection);
+
 	~FImplBackEndUEAndORT();
 
 	static void WarnAndSetDeviceToCPUIfDX12NotEnabled(ENeuralDeviceType& InOutDeviceType, const bool bInShouldOpenMessageLog);
@@ -53,11 +56,6 @@ public:
 		std::atomic<bool>& bInOutIsBackgroundThreadRunning, FCriticalSection& InOutResoucesCriticalSection, TArray<bool>& OutAreInputTensorSizesVariable,
 		const TArray<uint8>& InModelReadFromFileInBytes, const FString& InModelFullFilePath, const ENeuralDeviceType InDeviceType,
 		const ENeuralDeviceType InInputDeviceType, const ENeuralDeviceType InOutputDeviceType);
-	
-#ifdef WITH_UE_AND_ORT_SUPPORT
-	FImplBackEndUEAndORT(FOnAsyncRunCompleted& InOutOnAsyncRunCompletedDelegate, std::atomic<bool>& bInIsBackgroundThreadRunning,
-		FCriticalSection& InResoucesCriticalSection);
-#endif
 
 	void Run(const ENeuralNetworkSynchronousMode InSynchronousMode, const ENeuralDeviceType InDeviceType, const ENeuralDeviceType InInputDeviceType,
 		const ENeuralDeviceType InOutputDeviceType);
