Fix crash during cook when AssetImportData is null in NeuralNetwork

diff --git a/Engine/Plugins/Experimental/NeuralNetworkInference/Source/NeuralNetworkInference/Private/NeuralNetwork.cpp b/Engine/Plugins/Experimental/NeuralNetworkInference/Source/NeuralNetworkInference/Private/NeuralNetwork.cpp
--- a/Engine/Plugins/Experimental/NeuralNetworkInference/Source/NeuralNetworkInference/Private/NeuralNetwork.cpp
+++ b/Engine/Plugins/Experimental/NeuralNetworkInference/Source/NeuralNetworkInference/Private/NeuralNetwork.cpp
@@ -623,8 +623,8 @@ UAssetImportData* UNeuralNetwork::GetAndMaybeCreateAssetImportData()
 
 void UNeuralNetwork::ReimportAssetFromEditorData()
 {
-	//Get the re-import filename
-	const FString ImportedFilename = AssetImportData->GetFirstFilename();
+	//Get the re-import filename if possible
+	const FString ImportedFilename = AssetImportData ? AssetImportData->GetFirstFilename() : FString();
 	if (ImportedFilename.Len() > 0)
 	{
 		// Ensure that the file provided by the path exists
