Fixed RedirectCoutAndCerrToUeLog not in the right places for UNN.

#rnx

diff --git a/Engine/Plugins/Experimental/NeuralNetworkInference/Source/NeuralNetworkInference/Private/NeuralNetworkImplBackEndUEAndORT.cpp b/Engine/Plugins/Experimental/NeuralNetworkInference/Source/NeuralNetworkInference/Private/NeuralNetworkImplBackEndUEAndORT.cpp
--- a/Engine/Plugins/Experimental/NeuralNetworkInference/Source/NeuralNetworkInference/Private/NeuralNetworkImplBackEndUEAndORT.cpp
+++ b/Engine/Plugins/Experimental/NeuralNetworkInference/Source/NeuralNetworkInference/Private/NeuralNetworkImplBackEndUEAndORT.cpp
@@ -247,6 +247,7 @@ void UNeuralNetwork::FImplBackEndUEAndORT::FNeuralNetworkAsyncTask::SetRunSessio
 
 void UNeuralNetwork::FImplBackEndUEAndORT::FNeuralNetworkAsyncTask::DoWork()
 {
+	const FRedirectCoutAndCerrToUeLog RedirectCoutAndCerrToUeLog;
 	if (SyncMode == ENeuralNetworkSynchronousMode::Asynchronous)
 	{
 		BackEnd->RunSessionAsync(DeviceType, InputDeviceType, OutputDeviceType);
@@ -342,9 +343,9 @@ bool UNeuralNetwork::FImplBackEndUEAndORT::Load(TSharedPtr<FImplBackEndUEAndORT>
 	try
 #endif //WITH_EDITOR
 	{
-		// Avoid multi-threaded crashes
 		const FRedirectCoutAndCerrToUeLog RedirectCoutAndCerrToUeLog;
 
+		// Avoid multi-threaded crashes
 		if (InOutImplBackEndUEAndORT.IsValid())
 		{
 			InOutImplBackEndUEAndORT->EnsureAsyncTaskCompletion(/*bShouldWarnIfNotDone*/true);
@@ -428,8 +429,6 @@ void UNeuralNetwork::FImplBackEndUEAndORT::Run(const ENeuralNetworkSynchronousMo
 	try
 #endif //WITH_EDITOR
 	{
-		const FRedirectCoutAndCerrToUeLog RedirectCoutAndCerrToUeLog;
-
 		EnsureAsyncTaskCompletion(/*bShouldWarnIfNotDone*/true);
 		NeuralNetworkAsyncTask->GetTask().SetRunSessionArgs(InSynchronousMode, InDeviceType, InInputDeviceType, InOutputDeviceType);
 
