Towards Async NNI added to NNI QA tests.
- Step 1: ModelAccTestRun unifies all different types of accuracy tests.

#rnx

diff --git a/Engine/Plugins/Experimental/NeuralNetworkInference/Source/NeuralNetworkInferenceQA/Private/ModelUnitTester.cpp b/Engine/Plugins/Experimental/NeuralNetworkInference/Source/NeuralNetworkInferenceQA/Private/ModelUnitTester.cpp
--- a/Engine/Plugins/Experimental/NeuralNetworkInference/Source/NeuralNetworkInferenceQA/Private/ModelUnitTester.cpp
+++ b/Engine/Plugins/Experimental/NeuralNetworkInference/Source/NeuralNetworkInferenceQA/Private/ModelUnitTester.cpp
@@ -169,12 +169,13 @@ bool FModelUnitTester::ModelLoadAccuracyAndSpeedTests(const FString& InProjectCo
 		const bool bShouldRunUEAndORTBackEnd = (InCPURepetitionsForUEAndORTBackEnd[ModelIndex] + InGPURepetitionsForUEAndORTBackEnd[ModelIndex] > 0);
 		if (bShouldRunUEAndORTBackEnd)
 		{
-			bDidGlobalTestPassed &= ModelAccuracyTest(Network, ENeuralBackEnd::UEAndORT, InInputArrayValues, CPUGroundTruths, GPUGroundTruths);
+			bDidGlobalTestPassed &= ModelAccuracyTest(Network, ENeuralNetworkSynchronousMode::Synchronous, ENeuralBackEnd::UEAndORT, InInputArrayValues, CPUGroundTruths, GPUGroundTruths);
+			//bDidGlobalTestPassed &= ModelAccuracyTest(Network, ENeuralNetworkSynchronousMode::Asynchronous, ENeuralBackEnd::UEAndORT, InInputArrayValues, CPUGroundTruths, GPUGroundTruths); // @todo: Fix this
 		}
 		const bool bShouldRunUEOnlyBackEnd = false; // (InCPURepetitionsForUEOnlyBackEnd[ModelIndex] * InCPURepetitionsForUEOnlyBackEnd[ModelIndex] > 0);
 		if (bShouldRunUEOnlyBackEnd)
 		{
-			bDidGlobalTestPassed &= ModelAccuracyTest(Network, ENeuralBackEnd::UEOnly, InInputArrayValues, CPUGroundTruths, GPUGroundTruths);
+			bDidGlobalTestPassed &= ModelAccuracyTest(Network, ENeuralNetworkSynchronousMode::Synchronous, ENeuralBackEnd::UEOnly, InInputArrayValues, CPUGroundTruths, GPUGroundTruths);
 		}
 
 		UE_LOG(LogNeuralNetworkInferenceQA, Display, TEXT("---------------------------------------------------------------------------------------------------------------------------------"));
@@ -188,11 +189,12 @@ bool FModelUnitTester::ModelLoadAccuracyAndSpeedTests(const FString& InProjectCo
 			UE_LOG(LogNeuralNetworkInferenceQA, Display, TEXT("-------------------- %s - Network %s Load and Run - %s"), *ModelName, *ModelType, *ModelFilePath);
 			if (bShouldRunUEAndORTBackEnd)
 			{
-				bDidGlobalTestPassed &= ModelAccuracyTest(NetworkONNXOrORTLoadTest(ModelFilePath), ENeuralBackEnd::UEAndORT, InInputArrayValues, CPUGroundTruths, GPUGroundTruths);
+				bDidGlobalTestPassed &= ModelAccuracyTest(NetworkONNXOrORTLoadTest(ModelFilePath), ENeuralNetworkSynchronousMode::Synchronous, ENeuralBackEnd::UEAndORT, InInputArrayValues, CPUGroundTruths, GPUGroundTruths);
+				//bDidGlobalTestPassed &= ModelAccuracyTest(NetworkONNXOrORTLoadTest(ModelFilePath), ENeuralNetworkSynchronousMode::Asynchronous, ENeuralBackEnd::UEAndORT, InInputArrayValues, CPUGroundTruths, GPUGroundTruths); // @todo: Fix this
 			}
 			if (bShouldRunUEOnlyBackEnd)
 			{
-				bDidGlobalTestPassed &= ModelAccuracyTest(NetworkONNXOrORTLoadTest(ModelFilePath), ENeuralBackEnd::UEOnly, InInputArrayValues, CPUGroundTruths, GPUGroundTruths);
+				bDidGlobalTestPassed &= ModelAccuracyTest(NetworkONNXOrORTLoadTest(ModelFilePath), ENeuralNetworkSynchronousMode::Synchronous, ENeuralBackEnd::UEOnly, InInputArrayValues, CPUGroundTruths, GPUGroundTruths);
 			}
 		}
 #else //WITH_EDITOR
@@ -287,7 +289,8 @@ UNeuralNetwork* FModelUnitTester::NetworkONNXOrORTLoadTest(const FString& InMode
 	return Network;
 }
 
-bool FModelUnitTester::ModelAccuracyTest(UNeuralNetwork* InOutNetwork, const ENeuralBackEnd InBackEnd, const TArray<float>& InInputArrayValues, const TArray<double>& InCPUGroundTruths, const TArray<double>& InGPUGroundTruths)
+bool FModelUnitTester::ModelAccuracyTest(UNeuralNetwork* InOutNetwork, const ENeuralNetworkSynchronousMode InSynchronousMode, const ENeuralBackEnd InBackEnd, const TArray<float>& InInputArrayValues,
+	const TArray<double>& InCPUGroundTruths, const TArray<double>& InGPUGroundTruths)
 {
 	// Sanity check
 	if (!InOutNetwork)
@@ -317,8 +320,12 @@ bool FModelUnitTester::ModelAccuracyTest(UNeuralNetwork* InOutNetwork, const ENe
 	const ENeuralDeviceType OriginalDeviceType = InOutNetwork->GetDeviceType();
 	const ENeuralDeviceType OriginalInputDeviceType = InOutNetwork->GetInputDeviceType();
 	const ENeuralDeviceType OriginalOutputDeviceType = InOutNetwork->GetOutputDeviceType();
+	const ENeuralNetworkSynchronousMode OriginalSynchronousMode = InOutNetwork->GetSynchronousMode();
 	const ENeuralBackEnd OriginalBackEnd = InOutNetwork->GetBackEnd();
 	
+	// Set (a)synchronous Mode
+	InOutNetwork->SetSynchronousMode(InSynchronousMode);
+
 	// Set back end
 	if (!InOutNetwork->SetBackEnd(InBackEnd))
 	{
@@ -333,44 +340,25 @@ bool FModelUnitTester::ModelAccuracyTest(UNeuralNetwork* InOutNetwork, const ENe
 	TArray<TArray<float>> CPUOutputs, CPUGPUCPUOutputs, CPUGPUGPUOutputs, GPUGPUCPUOutputs;
 	
 	// Input CPU + Network CPU + Output CPU
-	for (int32 Index = 0; Index < InputArrays.Num(); ++Index)
-	{
-		// CPU
-		InOutNetwork->SetDeviceType(/*DeviceType*/ENeuralDeviceType::CPU, /*InputDeviceType*/ENeuralDeviceType::CPU, /*OutputDeviceType*/ENeuralDeviceType::CPU);
-		InOutNetwork->SetInputFromArrayCopy(InputArrays[Index]);
-		InOutNetwork->Run();
-		CPUOutputs.Emplace(InOutNetwork->GetOutputTensor().GetArrayCopy<float>());
-	}
-	
+	ModelAccuracyTestRun(CPUOutputs, InOutNetwork, InSynchronousMode, InputArrays, /*DeviceType*/ENeuralDeviceType::CPU, /*InputDeviceType*/ENeuralDeviceType::CPU, /*OutputDeviceType*/ENeuralDeviceType::CPU);
+//InOutNetwork->SetSynchronousMode(OriginalSynchronousMode); // @todo: Fix this
+
 	if (InOutNetwork->IsGPUSupported())
 	{
 		// Input CPU + Network GPU + Output CPU
 		for (int32 Index = 0; Index < InputArrays.Num(); ++Index)
 		{
-			InOutNetwork->SetDeviceType(/*DeviceType*/ENeuralDeviceType::GPU, /*InputDeviceType*/ENeuralDeviceType::CPU, /*OutputDeviceType*/ENeuralDeviceType::CPU);
-			InOutNetwork->SetInputFromArrayCopy(InputArrays[Index]);
-			InOutNetwork->Run();
-			CPUGPUCPUOutputs.Emplace(InOutNetwork->GetOutputTensor().GetArrayCopy<float>());
+			ModelAccuracyTestRun(CPUGPUCPUOutputs, InOutNetwork, /*InSynchronousMode*/ENeuralNetworkSynchronousMode::Synchronous, InputArrays, /*DeviceType*/ENeuralDeviceType::GPU, /*InputDeviceType*/ENeuralDeviceType::CPU, /*OutputDeviceType*/ENeuralDeviceType::CPU);
 		}
-		
 		// Input CPU + Network GPU + Output GPU
 		for (int32 Index = 0; Index < InputArrays.Num(); ++Index)
 		{
-			InOutNetwork->SetDeviceType(/*DeviceType*/ENeuralDeviceType::GPU, /*InputDeviceType*/ENeuralDeviceType::CPU, /*OutputDeviceType*/ENeuralDeviceType::GPU);
-			InOutNetwork->SetInputFromArrayCopy(InputArrays[Index]);
-			InOutNetwork->Run();
-			InOutNetwork->OutputTensorsToCPU();
-			CPUGPUGPUOutputs.Emplace(InOutNetwork->GetOutputTensor().GetArrayCopy<float>());
+			ModelAccuracyTestRun(CPUGPUGPUOutputs, InOutNetwork, /*InSynchronousMode*/ENeuralNetworkSynchronousMode::Synchronous, InputArrays, /*DeviceType*/ENeuralDeviceType::GPU, /*InputDeviceType*/ENeuralDeviceType::CPU, /*OutputDeviceType*/ENeuralDeviceType::GPU);
 		}
-
 		// Input GPU + Network GPU + Output CPU
 		for (int32 Index = 0; Index < InputArrays.Num(); ++Index)
 		{
-			InOutNetwork->SetDeviceType(/*DeviceType*/ENeuralDeviceType::GPU, /*InputDeviceType*/ENeuralDeviceType::GPU, /*OutputDeviceType*/ENeuralDeviceType::CPU);
-			InOutNetwork->SetInputFromArrayCopy(InputArrays[Index]);
-			InOutNetwork->InputTensorsToGPU();
-			InOutNetwork->Run();
-			GPUGPUCPUOutputs.Emplace(InOutNetwork->GetOutputTensor().GetArrayCopy<float>());
+			ModelAccuracyTestRun(GPUGPUCPUOutputs, InOutNetwork, /*InSynchronousMode*/ENeuralNetworkSynchronousMode::Synchronous, InputArrays, /*DeviceType*/ENeuralDeviceType::GPU, /*InputDeviceType*/ENeuralDeviceType::GPU, /*OutputDeviceType*/ENeuralDeviceType::CPU);
 		}
 	}
 	else
@@ -416,8 +404,10 @@ bool FModelUnitTester::ModelAccuracyTest(UNeuralNetwork* InOutNetwork, const ENe
 			UE_LOG(LogNeuralNetworkInferenceQA, Display, TEXT("CPUGPUAvgL1NormDiff (%fe-6) < 5e-6 might have failed."), CPUGPUAvgL1NormDiff);
 			UE_LOG(LogNeuralNetworkInferenceQA, Display, TEXT("GPUGPUInputAvgL1NormDiff (%fe-12) < 1e-12 might have failed."), GPUGPUInputAvgL1NormDiff);
 			UE_LOG(LogNeuralNetworkInferenceQA, Display, TEXT("GPUGPUOutputAvgL1NormDiff (%fe-12) < 1e-12 might have failed."), GPUGPUOutputAvgL1NormDiff);
-			UE_LOG(LogNeuralNetworkInferenceQA, Display, TEXT("FastCPUAvgL1NormDiff (%fe-7) < 30e-7 might have failed (~30 times the float precision).\nCPUOutput = %s."), FastCPUAvgL1NormDiff, *FNeuralTensor(CPUOutputs[Index], OutputSizes).ToString(MaxNumberElementsToDisplay));
-			UE_LOG(LogNeuralNetworkInferenceQA, Display, TEXT("FastGPUAvgL1NormDiff (%fe-7) < 30e-7 might have failed (~30 times the float precision).\nCPUGPUCPUOutput = %s."), FastGPUAvgL1NormDiff, *FNeuralTensor(CPUGPUCPUOutput, OutputSizes).ToString(MaxNumberElementsToDisplay));
+			UE_LOG(LogNeuralNetworkInferenceQA, Display, TEXT("FastCPUAvgL1NormDiff (%fe-7) < 30e-7 might have failed (~30 times the float precision).\nCPUOutput = %s."),
+				FastCPUAvgL1NormDiff, *FNeuralTensor(CPUOutputs[Index], OutputSizes).ToString(MaxNumberElementsToDisplay));
+			UE_LOG(LogNeuralNetworkInferenceQA, Display, TEXT("FastGPUAvgL1NormDiff (%fe-7) < 30e-7 might have failed (~30 times the float precision).\nCPUGPUCPUOutput = %s."),
+				FastGPUAvgL1NormDiff, *FNeuralTensor(CPUGPUCPUOutput, OutputSizes).ToString(MaxNumberElementsToDisplay));
 			UE_LOG(LogNeuralNetworkInferenceQA, Display, TEXT("Input = %s"),
 				*FNeuralTensor(InOutNetwork->GetInputTensor().GetArrayCopy<float>(), InputSizes).ToString(MaxNumberElementsToDisplay));
 			UE_LOG(LogNeuralNetworkInferenceQA, Display, TEXT("CPUOutput = %s"), *FNeuralTensor(CPUOutputs[Index], OutputSizes).ToString(MaxNumberElementsToDisplay));
@@ -431,11 +421,51 @@ bool FModelUnitTester::ModelAccuracyTest(UNeuralNetwork* InOutNetwork, const ENe
 	}
 	// Reset to original network state
 	InOutNetwork->SetDeviceType(/*DeviceType*/OriginalDeviceType, /*InputDeviceType*/OriginalInputDeviceType, /*OutputDeviceType*/OriginalOutputDeviceType);
+	InOutNetwork->SetSynchronousMode(OriginalSynchronousMode);
 	InOutNetwork->SetBackEnd(OriginalBackEnd);
 	// Test successful
 	return bDidGlobalTestPassed;
 }
 
+void FModelUnitTester::ModelAccuracyTestRun(TArray<TArray<float>>& OutOutputs, UNeuralNetwork* InOutNetwork, const ENeuralNetworkSynchronousMode InSynchronousMode,
+	const TArray<TArray<float>> InInputArrays, const ENeuralDeviceType InDeviceType, const ENeuralDeviceType InInputDeviceType, const ENeuralDeviceType InOutputDeviceType)
+{
+	std::atomic<bool> bDidAsyncNetworkFinished(false);
+	if (InSynchronousMode == ENeuralNetworkSynchronousMode::Asynchronous)
+	{
+		InOutNetwork->GetOnAsyncRunCompletedDelegate().BindLambda([&bDidAsyncNetworkFinished]()
+		{
+			bDidAsyncNetworkFinished = true;
+		});
+	}
+	for (int32 Index = 0; Index < InInputArrays.Num(); ++Index)
+	{
+		InOutNetwork->SetDeviceType(InDeviceType, InInputDeviceType, InOutputDeviceType);
+		InOutNetwork->SetInputFromArrayCopy(InInputArrays[Index]);
+		if (InInputDeviceType == ENeuralDeviceType::GPU)
+		{
+			InOutNetwork->InputTensorsToGPU();
+		}
+		InOutNetwork->Run();
+		// If async test, wait until inference is completed
+		if (InSynchronousMode == ENeuralNetworkSynchronousMode::Asynchronous)
+		{
+			while (!bDidAsyncNetworkFinished)
+			{
+				FPlatformProcess::Sleep(0.1e-3);
+			}
+		}
+		if (InOutputDeviceType == ENeuralDeviceType::GPU)
+		{
+			InOutNetwork->OutputTensorsToCPU();
+		}
+		OutOutputs.Emplace(InOutNetwork->GetOutputTensor().GetArrayCopy<float>());
+	}
+	if (InSynchronousMode == ENeuralNetworkSynchronousMode::Asynchronous)
+	{
+		InOutNetwork->GetOnAsyncRunCompletedDelegate().Unbind();
+	}
+}
 
 bool FModelUnitTester::ModelSpeedTest(const FString& InUAssetPath, const ENeuralDeviceType InDeviceType, const ENeuralBackEnd InBackEnd, const int32 InRepetitions)
 {
diff --git a/Engine/Plugins/Experimental/NeuralNetworkInference/Source/NeuralNetworkInferenceQA/Private/ModelUnitTester.h b/Engine/Plugins/Experimental/NeuralNetworkInference/Source/NeuralNetworkInferenceQA/Private/ModelUnitTester.h
--- a/Engine/Plugins/Experimental/NeuralNetworkInference/Source/NeuralNetworkInferenceQA/Private/ModelUnitTester.h
+++ b/Engine/Plugins/Experimental/NeuralNetworkInference/Source/NeuralNetworkInferenceQA/Private/ModelUnitTester.h
@@ -25,7 +25,10 @@ private:
 	static FString GetUAssetModelFilePath(const FString& InModelName, const FString& InModelZooRelativeDirectory);
 	static UNeuralNetwork* NetworkUassetLoadTest(const FString& InUAssetPath);
 	static UNeuralNetwork* NetworkONNXOrORTLoadTest(const FString& InModelFilePath);
-	static bool ModelAccuracyTest(UNeuralNetwork* InOutNetwork, const ENeuralBackEnd InBackEnd, const TArray<float>& InInputArrayValues, const TArray<double>& InCPUGroundTruths, const TArray<double>& InGPUGroundTruths);
+	static bool ModelAccuracyTest(UNeuralNetwork* InOutNetwork, const ENeuralNetworkSynchronousMode InSynchronousMode, const ENeuralBackEnd InBackEnd, const TArray<float>& InInputArrayValues,
+		const TArray<double>& InCPUGroundTruths, const TArray<double>& InGPUGroundTruths);
+	static void ModelAccuracyTestRun(TArray<TArray<float>>& OutOutputs, UNeuralNetwork* InOutNetwork, const ENeuralNetworkSynchronousMode InSynchronousMode, const TArray<TArray<float>> InInputArrays,
+		const ENeuralDeviceType InDeviceType, const ENeuralDeviceType InInputDeviceType, const ENeuralDeviceType InOutputDeviceType);
 	static bool ModelSpeedTest(const FString& InUAssetPath, const ENeuralDeviceType InDeviceType, const ENeuralBackEnd InBackEnd, const int32 InRepetitions);
 	static double GetAveragedL1Norm(const TArray<float>& InArray);
 	static double GetAveragedL1NormDiff(const TArray<float>& InArray1, const TArray<float>& InArray2);
