1. Removed bIsBackgroundThreadRunning from UNN (but moved to UEOnly back end that stills need it)
2. Removed IsReadyForFinishDestroy (instead added destructors to BackEndUEOnly/UEAndORT that block until freed)

#rnx

diff --git a/Engine/Plugins/Experimental/NeuralNetworkInference/Source/NeuralNetworkInference/Private/NeuralNetwork.cpp b/Engine/Plugins/Experimental/NeuralNetworkInference/Source/NeuralNetworkInference/Private/NeuralNetwork.cpp
--- a/Engine/Plugins/Experimental/NeuralNetworkInference/Source/NeuralNetworkInference/Private/NeuralNetwork.cpp
+++ b/Engine/Plugins/Experimental/NeuralNetworkInference/Source/NeuralNetworkInference/Private/NeuralNetwork.cpp
@@ -51,7 +51,6 @@ UNeuralNetwork::UNeuralNetwork()
 	, SynchronousMode(ENeuralNetworkSynchronousMode::Synchronous)
 	, BackEnd(ENeuralBackEnd::Auto)
 	, bIsLoaded(false)
-	, bIsBackgroundThreadRunning(false)
 	, BackEndForCurrentPlatform(FPrivateNeuralNetwork::SetBackEndForCurrentPlatform(BackEnd))
 {
 }
@@ -513,7 +512,7 @@ void UNeuralNetwork::Run()
 	// UEOnly
 	else if (BackEndForCurrentPlatform == ENeuralBackEnd::UEOnly)
 	{
-		ImplBackEndUEOnly->Run(OnAsyncRunCompletedDelegate, bIsBackgroundThreadRunning, SynchronousMode, DeviceType, InputDeviceType, OutputDeviceType);
+		ImplBackEndUEOnly->Run(OnAsyncRunCompletedDelegate, SynchronousMode, DeviceType, InputDeviceType, OutputDeviceType);
 	}
 	// Unknown
 	else
@@ -555,7 +554,7 @@ bool UNeuralNetwork::Load()
 	if (BackEndForCurrentPlatform == ENeuralBackEnd::UEAndORT)
 	{
 		UNeuralNetwork::FImplBackEndUEAndORT::WarnAndSetDeviceToCPUIfDX12NotEnabled(DeviceType, /*bShouldOpenMessageLog*/true);
-		bIsLoaded = UNeuralNetwork::FImplBackEndUEAndORT::Load(ImplBackEndUEAndORT, OnAsyncRunCompletedDelegate, bIsBackgroundThreadRunning, ResoucesCriticalSection, AreInputTensorSizesVariable,
+		bIsLoaded = UNeuralNetwork::FImplBackEndUEAndORT::Load(ImplBackEndUEAndORT, OnAsyncRunCompletedDelegate, ResoucesCriticalSection, AreInputTensorSizesVariable,
 			ModelReadFromFileInBytes, ModelFullFilePath, GetDeviceType(), GetInputDeviceType(), GetOutputDeviceType());
 	}
 	// UEOnly
@@ -678,8 +677,3 @@ void UNeuralNetwork::Serialize(FArchive& Archive)
 #endif // WITH_EDITORONLY_DATA
 	Super::Serialize(Archive);
 }
-
-bool UNeuralNetwork::IsReadyForFinishDestroy()
-{
-	return !bIsBackgroundThreadRunning;
-}
diff --git a/Engine/Plugins/Experimental/NeuralNetworkInference/Source/NeuralNetworkInference/Private/NeuralNetworkImplBackEndUEAndORT.cpp b/Engine/Plugins/Experimental/NeuralNetworkInference/Source/NeuralNetworkInference/Private/NeuralNetworkImplBackEndUEAndORT.cpp
--- a/Engine/Plugins/Experimental/NeuralNetworkInference/Source/NeuralNetworkInference/Private/NeuralNetworkImplBackEndUEAndORT.cpp
+++ b/Engine/Plugins/Experimental/NeuralNetworkInference/Source/NeuralNetworkInference/Private/NeuralNetworkImplBackEndUEAndORT.cpp
@@ -201,19 +201,19 @@ IDMLDevice* FPrivateImplBackEndUEAndORT::FDMLDeviceList::Add(ID3D12Device* Devic
 	}
 #endif
 
-	HRESULT res = DMLCreateDevice1(Device, DmlCreateFlags, DML_FEATURE_LEVEL_2_0, IID_PPV_ARGS(&DmlDevice));
+	HRESULT HResult = DMLCreateDevice1(Device, DmlCreateFlags, DML_FEATURE_LEVEL_2_0, IID_PPV_ARGS(&DmlDevice));
 
 	// Handle the case if Graphics Debug Tools are not installed
-	if (res == DXGI_ERROR_SDK_COMPONENT_MISSING)
+	if (HResult == DXGI_ERROR_SDK_COMPONENT_MISSING)
 	{
 		DmlCreateFlags &= ~DML_CREATE_DEVICE_FLAG_DEBUG;
 
-		res = DMLCreateDevice1(Device, DmlCreateFlags, DML_FEATURE_LEVEL_2_0, IID_PPV_ARGS(&DmlDevice));
+		HResult = DMLCreateDevice1(Device, DmlCreateFlags, DML_FEATURE_LEVEL_2_0, IID_PPV_ARGS(&DmlDevice));
 	}
 
 	if (!DmlDevice)
 	{
-		UE_LOG(LogNeuralNetworkInference, Warning, TEXT("FDMLDeviceList::Add(): Failed to create DML device, res=%0x."), res);
+		UE_LOG(LogNeuralNetworkInference, Warning, TEXT("FDMLDeviceList::Add(): Failed to create DML device, HResult=%0x."), HResult);
 		return nullptr;
 	}
 
@@ -264,11 +264,9 @@ void UNeuralNetwork::FImplBackEndUEAndORT::FNeuralNetworkAsyncTask::DoWork()
 /* FImplBackEndUEAndORT 'structors
  *****************************************************************************/
 
-UNeuralNetwork::FImplBackEndUEAndORT::FImplBackEndUEAndORT(FOnAsyncRunCompleted& InOutOnAsyncRunCompletedDelegate,
-	std::atomic<bool>& bInIsBackgroundThreadRunning, FCriticalSection& InResoucesCriticalSection)
+UNeuralNetwork::FImplBackEndUEAndORT::FImplBackEndUEAndORT(FOnAsyncRunCompleted& InOutOnAsyncRunCompletedDelegate, FCriticalSection& InResoucesCriticalSection)
 #ifdef WITH_UE_AND_ORT_SUPPORT
 	: OnAsyncRunCompletedDelegate(InOutOnAsyncRunCompletedDelegate)
-	, bIsBackgroundThreadRunning(bInIsBackgroundThreadRunning)
 	, ResoucesCriticalSection(InResoucesCriticalSection)
 #endif
 {
@@ -334,7 +332,7 @@ bool UNeuralNetwork::FImplBackEndUEAndORT::IsGPUSupported()
 }
 
 bool UNeuralNetwork::FImplBackEndUEAndORT::Load(TSharedPtr<FImplBackEndUEAndORT>& InOutImplBackEndUEAndORT,
-	FOnAsyncRunCompleted& InOutOnAsyncRunCompletedDelegate, std::atomic<bool>& bInOutIsBackgroundThreadRunning, FCriticalSection& InOutResoucesCriticalSection,
+	FOnAsyncRunCompleted& InOutOnAsyncRunCompletedDelegate, FCriticalSection& InOutResoucesCriticalSection,
 	TArray<bool>& OutAreInputTensorSizesVariable, const TArray<uint8>& InModelReadFromFileInBytes, const FString& InModelFullFilePath,
 	const ENeuralDeviceType InDeviceType, const ENeuralDeviceType InInputDeviceType, const ENeuralDeviceType InOutputDeviceType)
 {
@@ -353,7 +351,7 @@ bool UNeuralNetwork::FImplBackEndUEAndORT::Load(TSharedPtr<FImplBackEndUEAndORT>
 
 		// Initialize and configure InOutImplBackEndUEAndORT
 		if (!UNeuralNetwork::FImplBackEndUEAndORT::InitializedAndConfigureMembers(InOutImplBackEndUEAndORT, InOutOnAsyncRunCompletedDelegate,
-				bInOutIsBackgroundThreadRunning, InOutResoucesCriticalSection, InModelFullFilePath, InDeviceType))
+				InOutResoucesCriticalSection, InModelFullFilePath, InDeviceType))
 		{
 			UE_LOG(LogNeuralNetworkInference, Warning, TEXT("FImplBackEndUEAndORT::Load(): InitializedAndConfigureMembers failed."));
 			return false;
@@ -439,7 +437,6 @@ void UNeuralNetwork::FImplBackEndUEAndORT::Run(const ENeuralNetworkSynchronousMo
 		}
 		else if (InSynchronousMode == ENeuralNetworkSynchronousMode::Asynchronous)
 		{
-			bIsBackgroundThreadRunning = true;
 			NeuralNetworkAsyncTask->StartBackgroundTask();
 		}
 		else
@@ -480,14 +477,12 @@ void UNeuralNetwork::FImplBackEndUEAndORT::EnsureAsyncTaskCompletion(const bool
 }
 
 bool UNeuralNetwork::FImplBackEndUEAndORT::InitializedAndConfigureMembers(TSharedPtr<FImplBackEndUEAndORT>& InOutImplBackEndUEAndORT,
-	FOnAsyncRunCompleted& InOutOnAsyncRunCompletedDelegate, std::atomic<bool>& bInOutIsBackgroundThreadRunning, FCriticalSection& InOutResoucesCriticalSection,
-	const FString& InModelFullFilePath, const ENeuralDeviceType InDeviceType)
+	FOnAsyncRunCompleted& InOutOnAsyncRunCompletedDelegate, FCriticalSection& InOutResoucesCriticalSection, const FString& InModelFullFilePath, const ENeuralDeviceType InDeviceType)
 {
 	// Initialize InOutImplBackEndUEAndORT
 	if (!InOutImplBackEndUEAndORT.IsValid())
 	{
-		InOutImplBackEndUEAndORT = MakeShared<FImplBackEndUEAndORT>(InOutOnAsyncRunCompletedDelegate, bInOutIsBackgroundThreadRunning,
-			InOutResoucesCriticalSection);
+		InOutImplBackEndUEAndORT = MakeShared<FImplBackEndUEAndORT>(InOutOnAsyncRunCompletedDelegate, InOutResoucesCriticalSection);
 
 		// Set up ORT and create an environment
 		const char* const ModelFullFilePathCharPtr = TCHAR_TO_ANSI(*InModelFullFilePath);
@@ -936,7 +931,6 @@ void UNeuralNetwork::FImplBackEndUEAndORT::RunSessionAsync(const ENeuralDeviceTy
 	RunSessionImpl(InDeviceType, InInputDeviceType, InOutputDeviceType);
 
 	OnAsyncRunCompletedDelegate.ExecuteIfBound();
-	bIsBackgroundThreadRunning = false;
 }
 
 void UNeuralNetwork::FImplBackEndUEAndORT::RunSessionSync(const ENeuralDeviceType InDeviceType, const ENeuralDeviceType InInputDeviceType,
diff --git a/Engine/Plugins/Experimental/NeuralNetworkInference/Source/NeuralNetworkInference/Private/NeuralNetworkImplBackEndUEAndORT.h b/Engine/Plugins/Experimental/NeuralNetworkInference/Source/NeuralNetworkInference/Private/NeuralNetworkImplBackEndUEAndORT.h
--- a/Engine/Plugins/Experimental/NeuralNetworkInference/Source/NeuralNetworkInference/Private/NeuralNetworkImplBackEndUEAndORT.h
+++ b/Engine/Plugins/Experimental/NeuralNetworkInference/Source/NeuralNetworkInference/Private/NeuralNetworkImplBackEndUEAndORT.h
@@ -43,8 +43,7 @@ public:
 	TArray<FNeuralTensor> InputTensors;
 	TArray<FNeuralTensor> OutputTensors;
 
-	FImplBackEndUEAndORT(FOnAsyncRunCompleted& InOutOnAsyncRunCompletedDelegate, std::atomic<bool>& bInIsBackgroundThreadRunning,
-		FCriticalSection& InResoucesCriticalSection);
+	FImplBackEndUEAndORT(FOnAsyncRunCompleted& InOutOnAsyncRunCompletedDelegate, FCriticalSection& InResoucesCriticalSection);
 
 	~FImplBackEndUEAndORT();
 
@@ -53,9 +52,8 @@ public:
 	static bool IsGPUSupported();
 
 	static bool Load(TSharedPtr<FImplBackEndUEAndORT>& InOutImplBackEndUEAndORT, FOnAsyncRunCompleted& InOutOnAsyncRunCompletedDelegate,
-		std::atomic<bool>& bInOutIsBackgroundThreadRunning, FCriticalSection& InOutResoucesCriticalSection, TArray<bool>& OutAreInputTensorSizesVariable,
-		const TArray<uint8>& InModelReadFromFileInBytes, const FString& InModelFullFilePath, const ENeuralDeviceType InDeviceType,
-		const ENeuralDeviceType InInputDeviceType, const ENeuralDeviceType InOutputDeviceType);
+		FCriticalSection& InOutResoucesCriticalSection, TArray<bool>& OutAreInputTensorSizesVariable, const TArray<uint8>& InModelReadFromFileInBytes,
+		const FString& InModelFullFilePath, const ENeuralDeviceType InDeviceType, const ENeuralDeviceType InInputDeviceType, const ENeuralDeviceType InOutputDeviceType);
 
 	void Run(const ENeuralNetworkSynchronousMode InSynchronousMode, const ENeuralDeviceType InDeviceType, const ENeuralDeviceType InInputDeviceType,
 		const ENeuralDeviceType InOutputDeviceType);
@@ -64,7 +62,6 @@ public:
 private:
 	/** Async support */
 	FOnAsyncRunCompleted& OnAsyncRunCompletedDelegate;
-	std::atomic<bool>& bIsBackgroundThreadRunning;
 	FCriticalSection& ResoucesCriticalSection;
 	/** Network-related variables */
 	TUniquePtr<Ort::Env> Environment;
@@ -110,8 +107,7 @@ private:
 
 	void EnsureAsyncTaskCompletion(const bool bShouldWarnIfNotDone) const;
 
-	static bool InitializedAndConfigureMembers(TSharedPtr<FImplBackEndUEAndORT>& InOutImplBackEndUEAndORT,
-		FOnAsyncRunCompleted& InOutOnAsyncRunCompletedDelegate, std::atomic<bool>& bInOutIsBackgroundThreadRunning,
+	static bool InitializedAndConfigureMembers(TSharedPtr<FImplBackEndUEAndORT>& InOutImplBackEndUEAndORT, FOnAsyncRunCompleted& InOutOnAsyncRunCompletedDelegate,
 		FCriticalSection& InOutResoucesCriticalSection, const FString& InModelFullFilePath, const ENeuralDeviceType InDeviceType);
 
 	bool ConfigureMembers(const ENeuralDeviceType InDeviceType);
diff --git a/Engine/Plugins/Experimental/NeuralNetworkInference/Source/NeuralNetworkInference/Private/NeuralNetworkImplBackEndUEOnly.cpp b/Engine/Plugins/Experimental/NeuralNetworkInference/Source/NeuralNetworkInference/Private/NeuralNetworkImplBackEndUEOnly.cpp
--- a/Engine/Plugins/Experimental/NeuralNetworkInference/Source/NeuralNetworkInference/Private/NeuralNetworkImplBackEndUEOnly.cpp
+++ b/Engine/Plugins/Experimental/NeuralNetworkInference/Source/NeuralNetworkInference/Private/NeuralNetworkImplBackEndUEOnly.cpp
@@ -10,6 +10,19 @@
 
 
 
+/* FImplBackEndUEOnly destructor
+ *****************************************************************************/
+
+UNeuralNetwork::FImplBackEndUEOnly::~FImplBackEndUEOnly()
+{
+	while (bIsBackgroundThreadRunning)
+	{
+		FPlatformProcess::Sleep(0.1e-3);
+	}
+}
+
+
+
 /* FImplBackEndUEOnly public functions
  *****************************************************************************/
 
@@ -74,7 +87,8 @@ bool UNeuralNetwork::FImplBackEndUEOnly::Load(TSharedPtr<FImplBackEndUEOnly>& In
 //	return (InOutImplBackEndUEOnly->Operators.Num() > 0 && InOutImplBackEndUEOnly->TensorManager.IsLoaded());
 //}
 
-void UNeuralNetwork::FImplBackEndUEOnly::Run(FOnAsyncRunCompleted& InOutOnAsyncRunCompletedDelegate, std::atomic<bool>& bInIsBackgroundThreadRunning, const ENeuralNetworkSynchronousMode InSynchronousMode, const ENeuralDeviceType InDeviceType, const ENeuralDeviceType InInputDeviceType, const ENeuralDeviceType InOutputDeviceType)
+void UNeuralNetwork::FImplBackEndUEOnly::Run(FOnAsyncRunCompleted& InOutOnAsyncRunCompletedDelegate, const ENeuralNetworkSynchronousMode InSynchronousMode, const ENeuralDeviceType InDeviceType,
+	const ENeuralDeviceType InInputDeviceType, const ENeuralDeviceType InOutputDeviceType)
 {
 	// Run UNeuralNetwork::UEOnly
 	if (Operators.Num() > 0)
@@ -106,7 +120,7 @@ void UNeuralNetwork::FImplBackEndUEOnly::Run(FOnAsyncRunCompleted& InOutOnAsyncR
 
 			// On RHI thread
 			ENQUEUE_RENDER_COMMAND(UNeuralNetwork_UEOnly_Run_RenderThread)(
-				[this, &InOutOnAsyncRunCompletedDelegate, &bInIsBackgroundThreadRunning, InSynchronousMode, InInputDeviceType, InOutputDeviceType](FRHICommandListImmediate& RHICmdList)
+				[this, &InOutOnAsyncRunCompletedDelegate, InSynchronousMode, InInputDeviceType, InOutputDeviceType](FRHICommandListImmediate& RHICmdList)
 				{
 					FRDGBuilder GraphBuilder(RHICmdList, RDG_EVENT_NAME("FImplBackEndUEOnly::Run()"));
 
@@ -196,14 +210,14 @@ void UNeuralNetwork::FImplBackEndUEOnly::Run(FOnAsyncRunCompleted& InOutOnAsyncR
 					// Broadcast delegates (from the render thread)
 					if (InSynchronousMode == ENeuralNetworkSynchronousMode::Asynchronous)
 					{
-						bInIsBackgroundThreadRunning = true; // This will be done right away
+						bIsBackgroundThreadRunning = true; // This will be done right away
 						GraphBuilder.AddPass(
 							RDG_EVENT_NAME("Async delegate broadcast"),
 							ERDGPassFlags::None,
-							[&InOutOnAsyncRunCompletedDelegate, &bInIsBackgroundThreadRunning](FRHICommandListImmediate& RHICmdList)
+							[this, &InOutOnAsyncRunCompletedDelegate](FRHICommandListImmediate& RHICmdList)
 						{
 							InOutOnAsyncRunCompletedDelegate.ExecuteIfBound();
-							bInIsBackgroundThreadRunning = false; // This will be done in the render graph after all NNI operators have been run
+							bIsBackgroundThreadRunning = false; // This will be done in the render graph after all NNI operators have been run
 						});
 					}
 
diff --git a/Engine/Plugins/Experimental/NeuralNetworkInference/Source/NeuralNetworkInference/Private/NeuralNetworkImplBackEndUEOnly.h b/Engine/Plugins/Experimental/NeuralNetworkInference/Source/NeuralNetworkInference/Private/NeuralNetworkImplBackEndUEOnly.h
--- a/Engine/Plugins/Experimental/NeuralNetworkInference/Source/NeuralNetworkInference/Private/NeuralNetworkImplBackEndUEOnly.h
+++ b/Engine/Plugins/Experimental/NeuralNetworkInference/Source/NeuralNetworkInference/Private/NeuralNetworkImplBackEndUEOnly.h
@@ -14,6 +14,13 @@
 
 struct UNeuralNetwork::FImplBackEndUEOnly
 {
+	/**
+	 * Whether an inference pass (i.e., Run) is happening.
+	 * This variable is thread safe as long as only "Run" modifies it. Other functions can safely read it at any time.
+	 * If other functions outside of Run() have to modify it, consider using a mutex with a bool rather than just a std::atomic<bool>.
+	 */
+	std::atomic<bool> bIsBackgroundThreadRunning;
+
 	/**
 	 * It should always be false when loaded from uasset (FNeuralTensors are not auto-loaded to GPU).
 	 */
@@ -32,12 +39,12 @@ struct UNeuralNetwork::FImplBackEndUEOnly
 	 */
 	TArray<TSharedPtr<FNeuralOperator>> Operators;
 
+	~FImplBackEndUEOnly();
+
 	static bool Load(TSharedPtr<FImplBackEndUEOnly>& InOutImplBackEndUEOnly, const TArray<uint8>& InModelReadFromFileInBytes);
 	static bool Load(TSharedPtr<FImplBackEndUEOnly>& InOutImplBackEndUEOnly, FNeuralTensorManager& InTensorManager, const TArray<TSharedPtr<FNeuralOperator>>& InOperators);
 
-	//static bool Load(TSharedPtr<FImplBackEndUEOnly>& InOutImplBackEndUEOnly, FNeuralTensorManager& InTensorManager, const TArray<TSharedPtr<FNeuralOperator>>& InOperators);
-
-	void Run(FOnAsyncRunCompleted& InOutOnAsyncRunCompletedDelegate, std::atomic<bool>& bInIsBackgroundThreadRunning, const ENeuralNetworkSynchronousMode InSynchronousMode, const ENeuralDeviceType InDeviceType, const ENeuralDeviceType InInputDeviceType, const ENeuralDeviceType InOutputDeviceType);
+	void Run(FOnAsyncRunCompleted& InOutOnAsyncRunCompletedDelegate, const ENeuralNetworkSynchronousMode InSynchronousMode, const ENeuralDeviceType InDeviceType, const ENeuralDeviceType InInputDeviceType, const ENeuralDeviceType InOutputDeviceType);
 
 private:
 	/**
diff --git a/Engine/Plugins/Experimental/NeuralNetworkInference/Source/NeuralNetworkInference/Public/NeuralNetwork.h b/Engine/Plugins/Experimental/NeuralNetworkInference/Source/NeuralNetworkInference/Public/NeuralNetwork.h
--- a/Engine/Plugins/Experimental/NeuralNetworkInference/Source/NeuralNetworkInference/Public/NeuralNetwork.h
+++ b/Engine/Plugins/Experimental/NeuralNetworkInference/Source/NeuralNetworkInference/Public/NeuralNetwork.h
@@ -239,13 +239,6 @@ protected:
 private:
 	bool bIsLoaded;
 
-	/**
-	 * Whether an inference pass (i.e., Run) is happening.
-	 * This variable is thread safe as long as only "Run" modifies it. Other functions can safely read it at any time.
-	 * If other functions outside of Run() have to modify it, consider using a mutex with a bool rather than just a std::atomic<bool>.
-	 */
-	std::atomic<bool> bIsBackgroundThreadRunning;
-
 	/**
 	 * Critical section (mutex) used to avoid issues or crashes due to the asynchronous Run being run at the same time than any other non-const class function.
 	 * @see UNeuralNetwork::Run().
@@ -333,6 +326,5 @@ private:
 	virtual void PostInitProperties() override;
 	virtual void PostLoad() override;
 	virtual void Serialize(FArchive& Archive) override;
-	virtual bool IsReadyForFinishDestroy() override;
 	//~End of UObject interface
 };
