Async NNI added to NNI QA tests for CPU

#rnx

diff --git a/Engine/Plugins/Experimental/NeuralNetworkInference/Source/NeuralNetworkInferenceQA/Private/ModelUnitTester.cpp b/Engine/Plugins/Experimental/NeuralNetworkInference/Source/NeuralNetworkInferenceQA/Private/ModelUnitTester.cpp
--- a/Engine/Plugins/Experimental/NeuralNetworkInference/Source/NeuralNetworkInferenceQA/Private/ModelUnitTester.cpp
+++ b/Engine/Plugins/Experimental/NeuralNetworkInference/Source/NeuralNetworkInferenceQA/Private/ModelUnitTester.cpp
@@ -170,7 +170,7 @@ bool FModelUnitTester::ModelLoadAccuracyAndSpeedTests(const FString& InProjectCo
 		if (bShouldRunUEAndORTBackEnd)
 		{
 			bDidGlobalTestPassed &= ModelAccuracyTest(Network, ENeuralNetworkSynchronousMode::Synchronous, ENeuralBackEnd::UEAndORT, InInputArrayValues, CPUGroundTruths, GPUGroundTruths);
-			//bDidGlobalTestPassed &= ModelAccuracyTest(Network, ENeuralNetworkSynchronousMode::Asynchronous, ENeuralBackEnd::UEAndORT, InInputArrayValues, CPUGroundTruths, GPUGroundTruths); // @todo: Fix this
+			bDidGlobalTestPassed &= ModelAccuracyTest(Network, ENeuralNetworkSynchronousMode::Asynchronous, ENeuralBackEnd::UEAndORT, InInputArrayValues, CPUGroundTruths, GPUGroundTruths);
 		}
 		const bool bShouldRunUEOnlyBackEnd = false; // (InCPURepetitionsForUEOnlyBackEnd[ModelIndex] * InCPURepetitionsForUEOnlyBackEnd[ModelIndex] > 0);
 		if (bShouldRunUEOnlyBackEnd)
@@ -190,7 +190,7 @@ bool FModelUnitTester::ModelLoadAccuracyAndSpeedTests(const FString& InProjectCo
 			if (bShouldRunUEAndORTBackEnd)
 			{
 				bDidGlobalTestPassed &= ModelAccuracyTest(NetworkONNXOrORTLoadTest(ModelFilePath), ENeuralNetworkSynchronousMode::Synchronous, ENeuralBackEnd::UEAndORT, InInputArrayValues, CPUGroundTruths, GPUGroundTruths);
-				//bDidGlobalTestPassed &= ModelAccuracyTest(NetworkONNXOrORTLoadTest(ModelFilePath), ENeuralNetworkSynchronousMode::Asynchronous, ENeuralBackEnd::UEAndORT, InInputArrayValues, CPUGroundTruths, GPUGroundTruths); // @todo: Fix this
+				bDidGlobalTestPassed &= ModelAccuracyTest(NetworkONNXOrORTLoadTest(ModelFilePath), ENeuralNetworkSynchronousMode::Asynchronous, ENeuralBackEnd::UEAndORT, InInputArrayValues, CPUGroundTruths, GPUGroundTruths);
 			}
 			if (bShouldRunUEOnlyBackEnd)
 			{
@@ -340,32 +340,23 @@ bool FModelUnitTester::ModelAccuracyTest(UNeuralNetwork* InOutNetwork, const ENe
 	TArray<TArray<float>> CPUOutputs, CPUGPUCPUOutputs, CPUGPUGPUOutputs, GPUGPUCPUOutputs;
 	
 	// Input CPU + Network CPU + Output CPU
-	ModelAccuracyTestRun(CPUOutputs, InOutNetwork, InSynchronousMode, InputArrays, /*DeviceType*/ENeuralDeviceType::CPU, /*InputDeviceType*/ENeuralDeviceType::CPU, /*OutputDeviceType*/ENeuralDeviceType::CPU);
-//InOutNetwork->SetSynchronousMode(OriginalSynchronousMode); // @todo: Fix this
+	ModelAccuracyTestRun(CPUOutputs, InOutNetwork, InputArrays, /*DeviceType*/ENeuralDeviceType::CPU, /*InputDeviceType*/ENeuralDeviceType::CPU, /*OutputDeviceType*/ENeuralDeviceType::CPU);
+InOutNetwork->SetSynchronousMode(ENeuralNetworkSynchronousMode::Synchronous); // @todo: Fix this
 
 	if (InOutNetwork->IsGPUSupported())
 	{
 		// Input CPU + Network GPU + Output CPU
-		for (int32 Index = 0; Index < InputArrays.Num(); ++Index)
-		{
-			ModelAccuracyTestRun(CPUGPUCPUOutputs, InOutNetwork, /*InSynchronousMode*/ENeuralNetworkSynchronousMode::Synchronous, InputArrays, /*DeviceType*/ENeuralDeviceType::GPU, /*InputDeviceType*/ENeuralDeviceType::CPU, /*OutputDeviceType*/ENeuralDeviceType::CPU);
-		}
+		ModelAccuracyTestRun(CPUGPUCPUOutputs, InOutNetwork, InputArrays, /*DeviceType*/ENeuralDeviceType::GPU, /*InputDeviceType*/ENeuralDeviceType::CPU, /*OutputDeviceType*/ENeuralDeviceType::CPU);
 		// Input CPU + Network GPU + Output GPU
-		for (int32 Index = 0; Index < InputArrays.Num(); ++Index)
-		{
-			ModelAccuracyTestRun(CPUGPUGPUOutputs, InOutNetwork, /*InSynchronousMode*/ENeuralNetworkSynchronousMode::Synchronous, InputArrays, /*DeviceType*/ENeuralDeviceType::GPU, /*InputDeviceType*/ENeuralDeviceType::CPU, /*OutputDeviceType*/ENeuralDeviceType::GPU);
-		}
+		ModelAccuracyTestRun(CPUGPUGPUOutputs, InOutNetwork, InputArrays, /*DeviceType*/ENeuralDeviceType::GPU, /*InputDeviceType*/ENeuralDeviceType::CPU, /*OutputDeviceType*/ENeuralDeviceType::GPU);
 		// Input GPU + Network GPU + Output CPU
-		for (int32 Index = 0; Index < InputArrays.Num(); ++Index)
-		{
-			ModelAccuracyTestRun(GPUGPUCPUOutputs, InOutNetwork, /*InSynchronousMode*/ENeuralNetworkSynchronousMode::Synchronous, InputArrays, /*DeviceType*/ENeuralDeviceType::GPU, /*InputDeviceType*/ENeuralDeviceType::GPU, /*OutputDeviceType*/ENeuralDeviceType::CPU);
-		}
+		ModelAccuracyTestRun(GPUGPUCPUOutputs, InOutNetwork, InputArrays, /*DeviceType*/ENeuralDeviceType::GPU, /*InputDeviceType*/ENeuralDeviceType::GPU, /*OutputDeviceType*/ENeuralDeviceType::CPU);
 	}
 	else
 	{
 		UE_LOG(LogNeuralNetworkInferenceQA, Display, TEXT("-------------------- ModelAccuracyTest-UEAndORT-GPU skipped (DX12 not enabled)."));
 	}
-	
+
 	for (int32 Index = 0; Index < InputArrays.Num(); ++Index)
 	{
 		const TArray<float>& GPUGPUCPUOutput = (InOutNetwork->IsGPUSupported() ? GPUGPUCPUOutputs[Index] : CPUOutputs[Index]);
@@ -427,11 +418,12 @@ bool FModelUnitTester::ModelAccuracyTest(UNeuralNetwork* InOutNetwork, const ENe
 	return bDidGlobalTestPassed;
 }
 
-void FModelUnitTester::ModelAccuracyTestRun(TArray<TArray<float>>& OutOutputs, UNeuralNetwork* InOutNetwork, const ENeuralNetworkSynchronousMode InSynchronousMode,
-	const TArray<TArray<float>> InInputArrays, const ENeuralDeviceType InDeviceType, const ENeuralDeviceType InInputDeviceType, const ENeuralDeviceType InOutputDeviceType)
+void FModelUnitTester::ModelAccuracyTestRun(TArray<TArray<float>>& OutOutputs, UNeuralNetwork* InOutNetwork, const TArray<TArray<float>> InInputArrays,
+	const ENeuralDeviceType InDeviceType, const ENeuralDeviceType InInputDeviceType, const ENeuralDeviceType InOutputDeviceType)
 {
 	std::atomic<bool> bDidAsyncNetworkFinished(false);
-	if (InSynchronousMode == ENeuralNetworkSynchronousMode::Asynchronous)
+	const ENeuralNetworkSynchronousMode SynchronousMode = InOutNetwork->GetSynchronousMode();
+	if (SynchronousMode == ENeuralNetworkSynchronousMode::Asynchronous)
 	{
 		InOutNetwork->GetOnAsyncRunCompletedDelegate().BindLambda([&bDidAsyncNetworkFinished]()
 		{
@@ -440,6 +432,7 @@ void FModelUnitTester::ModelAccuracyTestRun(TArray<TArray<float>>& OutOutputs, U
 	}
 	for (int32 Index = 0; Index < InInputArrays.Num(); ++Index)
 	{
+		bDidAsyncNetworkFinished = false;
 		InOutNetwork->SetDeviceType(InDeviceType, InInputDeviceType, InOutputDeviceType);
 		InOutNetwork->SetInputFromArrayCopy(InInputArrays[Index]);
 		if (InInputDeviceType == ENeuralDeviceType::GPU)
@@ -448,7 +441,7 @@ void FModelUnitTester::ModelAccuracyTestRun(TArray<TArray<float>>& OutOutputs, U
 		}
 		InOutNetwork->Run();
 		// If async test, wait until inference is completed
-		if (InSynchronousMode == ENeuralNetworkSynchronousMode::Asynchronous)
+		if (SynchronousMode == ENeuralNetworkSynchronousMode::Asynchronous)
 		{
 			while (!bDidAsyncNetworkFinished)
 			{
@@ -461,7 +454,7 @@ void FModelUnitTester::ModelAccuracyTestRun(TArray<TArray<float>>& OutOutputs, U
 		}
 		OutOutputs.Emplace(InOutNetwork->GetOutputTensor().GetArrayCopy<float>());
 	}
-	if (InSynchronousMode == ENeuralNetworkSynchronousMode::Asynchronous)
+	if (SynchronousMode == ENeuralNetworkSynchronousMode::Asynchronous)
 	{
 		InOutNetwork->GetOnAsyncRunCompletedDelegate().Unbind();
 	}
diff --git a/Engine/Plugins/Experimental/NeuralNetworkInference/Source/NeuralNetworkInferenceQA/Private/ModelUnitTester.h b/Engine/Plugins/Experimental/NeuralNetworkInference/Source/NeuralNetworkInferenceQA/Private/ModelUnitTester.h
--- a/Engine/Plugins/Experimental/NeuralNetworkInference/Source/NeuralNetworkInferenceQA/Private/ModelUnitTester.h
+++ b/Engine/Plugins/Experimental/NeuralNetworkInference/Source/NeuralNetworkInferenceQA/Private/ModelUnitTester.h
@@ -16,7 +16,8 @@ private:
 	 */
 	static bool ModelLoadAccuracyAndSpeedTests(const FString& InProjectContentDir, const FString& InModelZooRelativeDirectory, const TArray<FString>& InModelNames,
 		const TArray<float>& InInputArrayValues, const TArray<TArray<double>>& InCPUGroundTruths, const TArray<TArray<double>>& InGPUGroundTruths,
-		const TArray<int32>& InCPURepetitionsForUEAndORT, const TArray<int32>& InGPURepetitionsForUEAndORTBackEnd, const TArray<int32>& InCPURepetitionsForUEOnly, const TArray<int32>& InGPURepetitionsForUEOnly);
+		const TArray<int32>& InCPURepetitionsForUEAndORT, const TArray<int32>& InGPURepetitionsForUEAndORTBackEnd, const TArray<int32>& InCPURepetitionsForUEOnly,
+		const TArray<int32>& InGPURepetitionsForUEOnly);
 	/**
 	 * Other auxiliary functions for GlobalTest().
 	 */
@@ -27,7 +28,7 @@ private:
 	static UNeuralNetwork* NetworkONNXOrORTLoadTest(const FString& InModelFilePath);
 	static bool ModelAccuracyTest(UNeuralNetwork* InOutNetwork, const ENeuralNetworkSynchronousMode InSynchronousMode, const ENeuralBackEnd InBackEnd, const TArray<float>& InInputArrayValues,
 		const TArray<double>& InCPUGroundTruths, const TArray<double>& InGPUGroundTruths);
-	static void ModelAccuracyTestRun(TArray<TArray<float>>& OutOutputs, UNeuralNetwork* InOutNetwork, const ENeuralNetworkSynchronousMode InSynchronousMode, const TArray<TArray<float>> InInputArrays,
+	static void ModelAccuracyTestRun(TArray<TArray<float>>& OutOutputs, UNeuralNetwork* InOutNetwork, const TArray<TArray<float>> InInputArrays,
 		const ENeuralDeviceType InDeviceType, const ENeuralDeviceType InInputDeviceType, const ENeuralDeviceType InOutputDeviceType);
 	static bool ModelSpeedTest(const FString& InUAssetPath, const ENeuralDeviceType InDeviceType, const ENeuralBackEnd InBackEnd, const int32 InRepetitions);
 	static double GetAveragedL1Norm(const TArray<float>& InArray);
