Async NNI not causing warnings.

_codereview Francisco.Vicente
#rnx

diff --git a/Engine/Plugins/Experimental/NeuralNetworkInference/Source/NeuralNetworkInference/Private/NeuralNetworkImplBackEndUEAndORT.cpp b/Engine/Plugins/Experimental/NeuralNetworkInference/Source/NeuralNetworkInference/Private/NeuralNetworkImplBackEndUEAndORT.cpp
--- a/Engine/Plugins/Experimental/NeuralNetworkInference/Source/NeuralNetworkInference/Private/NeuralNetworkImplBackEndUEAndORT.cpp
+++ b/Engine/Plugins/Experimental/NeuralNetworkInference/Source/NeuralNetworkInference/Private/NeuralNetworkImplBackEndUEAndORT.cpp
@@ -275,7 +275,7 @@ UNeuralNetwork::FImplBackEndUEAndORT::FImplBackEndUEAndORT(FOnAsyncRunCompleted&
 UNeuralNetwork::FImplBackEndUEAndORT::~FImplBackEndUEAndORT()
 {
 #ifdef WITH_UE_AND_ORT_SUPPORT
-	EnsureAsyncTaskCompletion(/*bShouldWarnIfNotDone*/false);
+	EnsureAsyncTaskCompletion();
 #endif //WITH_UE_AND_ORT_SUPPORT
 }
 
@@ -346,7 +346,7 @@ bool UNeuralNetwork::FImplBackEndUEAndORT::Load(TSharedPtr<FImplBackEndUEAndORT>
 		// Avoid multi-threaded crashes
 		if (InOutImplBackEndUEAndORT.IsValid())
 		{
-			InOutImplBackEndUEAndORT->EnsureAsyncTaskCompletion(/*bShouldWarnIfNotDone*/true);
+			InOutImplBackEndUEAndORT->EnsureAsyncTaskCompletion();
 		}
 
 		// Initialize and configure InOutImplBackEndUEAndORT
@@ -427,7 +427,7 @@ void UNeuralNetwork::FImplBackEndUEAndORT::Run(const ENeuralNetworkSynchronousMo
 	try
 #endif //WITH_EDITOR
 	{
-		EnsureAsyncTaskCompletion(/*bShouldWarnIfNotDone*/true);
+		EnsureAsyncTaskCompletion();
 		NeuralNetworkAsyncTask->GetTask().SetRunSessionArgs(InSynchronousMode, InDeviceType, InInputDeviceType, InOutputDeviceType);
 
 		// Run UNeuralNetwork
@@ -437,7 +437,7 @@ void UNeuralNetwork::FImplBackEndUEAndORT::Run(const ENeuralNetworkSynchronousMo
 		}
 		else if (InSynchronousMode == ENeuralNetworkSynchronousMode::Asynchronous)
 		{
-			NeuralNetworkAsyncTask->StartBackgroundTask();
+			NeuralNetworkAsyncTask->StartBackgroundTask(); // Alternative: (GThreadPool, EQueuedWorkPriority::Highest);
 		}
 		else
 		{
@@ -463,15 +463,10 @@ void UNeuralNetwork::FImplBackEndUEAndORT::Run(const ENeuralNetworkSynchronousMo
  *****************************************************************************/
 
 #ifdef WITH_UE_AND_ORT_SUPPORT
-void UNeuralNetwork::FImplBackEndUEAndORT::EnsureAsyncTaskCompletion(const bool bShouldWarnIfNotDone) const
+void UNeuralNetwork::FImplBackEndUEAndORT::EnsureAsyncTaskCompletion() const
 {
 	if (NeuralNetworkAsyncTask && !NeuralNetworkAsyncTask->IsDone())
 	{
-		if (bShouldWarnIfNotDone)
-		{
-			UE_LOG(LogNeuralNetworkInference, Warning,
-				TEXT("FImplBackEndUEAndORT::EnsureAsyncTaskCompletion(): Previous async run had not been completed. Blocking thread until it is completed."));
-		}
 		NeuralNetworkAsyncTask->EnsureCompletion(/*bDoWorkOnThisThreadIfNotStarted*/true);
 	}
 }
diff --git a/Engine/Plugins/Experimental/NeuralNetworkInference/Source/NeuralNetworkInference/Private/NeuralNetworkImplBackEndUEAndORT.h b/Engine/Plugins/Experimental/NeuralNetworkInference/Source/NeuralNetworkInference/Private/NeuralNetworkImplBackEndUEAndORT.h
--- a/Engine/Plugins/Experimental/NeuralNetworkInference/Source/NeuralNetworkInference/Private/NeuralNetworkImplBackEndUEAndORT.h
+++ b/Engine/Plugins/Experimental/NeuralNetworkInference/Source/NeuralNetworkInference/Private/NeuralNetworkImplBackEndUEAndORT.h
@@ -105,7 +105,7 @@ private:
 	};
 	TUniquePtr<FAsyncTask<FNeuralNetworkAsyncTask>> NeuralNetworkAsyncTask;
 
-	void EnsureAsyncTaskCompletion(const bool bShouldWarnIfNotDone) const;
+	void EnsureAsyncTaskCompletion() const;
 
 	static bool InitializedAndConfigureMembers(TSharedPtr<FImplBackEndUEAndORT>& InOutImplBackEndUEAndORT, FOnAsyncRunCompleted& InOutOnAsyncRunCompletedDelegate,
 		FCriticalSection& InOutResoucesCriticalSection, const FString& InModelFullFilePath, const ENeuralDeviceType InDeviceType);
