NeuralNetworkInference (NNI) plugin:
- Display which GPU is used by NNI
#rnx

diff --git a/Engine/Plugins/Experimental/NeuralNetworkInference/Source/NeuralNetworkInference/Private/NeuralNetworkImplBackEndUEAndORT.cpp b/Engine/Plugins/Experimental/NeuralNetworkInference/Source/NeuralNetworkInference/Private/NeuralNetworkImplBackEndUEAndORT.cpp
--- a/Engine/Plugins/Experimental/NeuralNetworkInference/Source/NeuralNetworkInference/Private/NeuralNetworkImplBackEndUEAndORT.cpp
+++ b/Engine/Plugins/Experimental/NeuralNetworkInference/Source/NeuralNetworkInference/Private/NeuralNetworkImplBackEndUEAndORT.cpp
@@ -524,18 +524,38 @@ bool UNeuralNetwork::FImplBackEndUEAndORT::ConfigureMembers(const ENeuralDeviceT
 		}
 
 		// Get adapter's D3D12 device that we would like to share with DirectML execution provider
-		// NOTE: For now we're only using first device that has Dadapter 0 and device 0
-		FD3D12DynamicRHI* Rhi = static_cast<FD3D12DynamicRHI*>(GDynamicRHI);
+		FD3D12DynamicRHI*			Rhi = static_cast<FD3D12DynamicRHI*>(GDynamicRHI);
+		FD3D12Adapter&				RhiAdapter = Rhi->GetAdapter(0);
+		const FD3D12AdapterDesc&	RhiAdapterDesc = RhiAdapter.GetDesc();
 
-		if (Rhi->GetNumAdapters() > 1 || Rhi->GetAdapter(0).GetDesc().NumDeviceNodes > 1)
+		UE_LOG(LogNeuralNetworkInference, Display, TEXT("FImplBackEndUEAndORT::ConfigureMembers(): NNI available adapters %d"), Rhi->GetNumAdapters());
+
+		for (int32 ic = 0; ic < Rhi->GetNumAdapters(); ++ic)
+		{
+			const FD3D12AdapterDesc& CurrAdapterDesc = Rhi->GetAdapter(ic).GetDesc();
+
+			UE_LOG(LogNeuralNetworkInference, Display,
+				TEXT("   Adapter [%d] Name:%s with LUID:%0x%0x"),
+				ic,
+				CurrAdapterDesc.Desc.Description,
+				CurrAdapterDesc.Desc.AdapterLuid.HighPart,
+				CurrAdapterDesc.Desc.AdapterLuid.LowPart);
+		}
+
+		UE_LOG(LogNeuralNetworkInference, Display,
+			TEXT("FImplBackEndUEAndORT::ConfigureMembers(): NNI using adapter %s with LUID:%0x%0x"),
+			RhiAdapterDesc.Desc.Description,
+			RhiAdapterDesc.Desc.AdapterLuid.HighPart,
+			RhiAdapterDesc.Desc.AdapterLuid.LowPart);
+
+		if (Rhi->GetNumAdapters() > 1 || RhiAdapterDesc.NumDeviceNodes > 1)
 		{
 			UE_LOG(LogNeuralNetworkInference, Warning,
-				TEXT("FImplBackEndUEAndORT::ConfigureMembers(): There are multiple (%d) adapters and/or multiple (%d) devices, using device at index 0."),
-				Rhi->GetNumAdapters(), Rhi->GetAdapter(0).GetDesc().NumDeviceNodes);
-			return false;
+				TEXT("FImplBackEndUEAndORT::ConfigureMembers(): There are multiple (%d) adapters and/or multiple (%d) devices, NNI is currently using only one adapter."),
+				Rhi->GetNumAdapters(), RhiAdapterDesc.NumDeviceNodes);
 		}
 
-		ID3D12Device* NativeDevice = Rhi->GetAdapter(0).GetD3DDevice();
+		ID3D12Device* NativeDevice = RhiAdapter.GetD3DDevice();
 
 		// Make sure that we have one DMLDevice per D3D12 device
 		IDMLDevice* DmlDevice = FPrivateImplBackEndUEAndORT::GetDMLDeviceThreadSafe(NativeDevice);
