Several minor changes:
1. NNI async callback optionally called from main thread by default.
2.  bIsLoaded, AreInputTensorSizesVariable, BackEndForCurrentPlatform as transient UPROPERTY's.
3. Added ENeuralNetworkDelegateThreadMode, GetOnAsyncRunCompletedDelegateMode, SetOnAsyncRunCompletedDelegateMode to control where to run callback.
4. Epic coding standard:
- Variables with descriptive names and starting with upper case.
- Code blocks in fewer lines.
- Reduced verbosity of RHI adapters from 3 lines to only 1.

#rnx

diff --git a/Engine/Plugins/Experimental/NeuralNetworkInference/Source/NeuralNetworkInference/Private/NeuralNetwork.cpp b/Engine/Plugins/Experimental/NeuralNetworkInference/Source/NeuralNetworkInference/Private/NeuralNetwork.cpp
--- a/Engine/Plugins/Experimental/NeuralNetworkInference/Source/NeuralNetworkInference/Private/NeuralNetwork.cpp
+++ b/Engine/Plugins/Experimental/NeuralNetworkInference/Source/NeuralNetworkInference/Private/NeuralNetwork.cpp
@@ -49,6 +49,7 @@ UNeuralNetwork::UNeuralNetwork()
 	, InputDeviceType(ENeuralDeviceType::CPU)
 	, OutputDeviceType(ENeuralDeviceType::CPU)
 	, SynchronousMode(ENeuralNetworkSynchronousMode::Synchronous)
+	, DelegateThreadMode(ENeuralNetworkDelegateThreadMode::GameThread)
 	, BackEnd(ENeuralBackEnd::Auto)
 	, bIsLoaded(false)
 	, BackEndForCurrentPlatform(FPrivateNeuralNetwork::SetBackEndForCurrentPlatform(BackEnd))
@@ -175,6 +176,17 @@ UNeuralNetwork::FOnAsyncRunCompleted& UNeuralNetwork::GetOnAsyncRunCompletedDele
 	return OnAsyncRunCompletedDelegate;
 }
 
+ENeuralNetworkDelegateThreadMode UNeuralNetwork::GetOnAsyncRunCompletedDelegateMode() const
+{
+	return DelegateThreadMode;
+}
+
+void UNeuralNetwork::SetOnAsyncRunCompletedDelegateMode(const ENeuralNetworkDelegateThreadMode InDelegateThreadMode)
+{
+	const FScopeLock ResourcesLock(&ResoucesCriticalSection);
+	DelegateThreadMode = InDelegateThreadMode;
+}
+
 ENeuralBackEnd UNeuralNetwork::GetBackEnd() const
 {
 	return BackEnd;
@@ -554,7 +566,7 @@ bool UNeuralNetwork::Load()
 	if (BackEndForCurrentPlatform == ENeuralBackEnd::UEAndORT)
 	{
 		UNeuralNetwork::FImplBackEndUEAndORT::WarnAndSetDeviceToCPUIfDX12NotEnabled(DeviceType, /*bShouldOpenMessageLog*/true);
-		bIsLoaded = UNeuralNetwork::FImplBackEndUEAndORT::Load(ImplBackEndUEAndORT, OnAsyncRunCompletedDelegate, ResoucesCriticalSection, AreInputTensorSizesVariable,
+		bIsLoaded = UNeuralNetwork::FImplBackEndUEAndORT::Load(ImplBackEndUEAndORT, OnAsyncRunCompletedDelegate, DelegateThreadMode, ResoucesCriticalSection, AreInputTensorSizesVariable,
 			ModelReadFromFileInBytes, ModelFullFilePath, GetDeviceType(), GetInputDeviceType(), GetOutputDeviceType());
 	}
 	// UEOnly
diff --git a/Engine/Plugins/Experimental/NeuralNetworkInference/Source/NeuralNetworkInference/Private/NeuralNetworkImplBackEndUEAndORT.cpp b/Engine/Plugins/Experimental/NeuralNetworkInference/Source/NeuralNetworkInference/Private/NeuralNetworkImplBackEndUEAndORT.cpp
--- a/Engine/Plugins/Experimental/NeuralNetworkInference/Source/NeuralNetworkInference/Private/NeuralNetworkImplBackEndUEAndORT.cpp
+++ b/Engine/Plugins/Experimental/NeuralNetworkInference/Source/NeuralNetworkInference/Private/NeuralNetworkImplBackEndUEAndORT.cpp
@@ -1,7 +1,7 @@
 // Copyright Epic Games, Inc. All Rights Reserved.
 
 #include "NeuralNetworkImplBackEndUEAndORT.h"
-
+#include "Async/Async.h"
 #include "NeuralNetworkInferenceUtils.h"
 #include "NeuralNetworkInferenceUtilsGPU.h"
 #include "RedirectCoutAndCerrToUeLog.h"
@@ -264,9 +264,11 @@ void UNeuralNetwork::FImplBackEndUEAndORT::FNeuralNetworkAsyncTask::DoWork()
 /* FImplBackEndUEAndORT 'structors
  *****************************************************************************/
 
-UNeuralNetwork::FImplBackEndUEAndORT::FImplBackEndUEAndORT(FOnAsyncRunCompleted& InOutOnAsyncRunCompletedDelegate, FCriticalSection& InResoucesCriticalSection)
+UNeuralNetwork::FImplBackEndUEAndORT::FImplBackEndUEAndORT(FOnAsyncRunCompleted& InOutOnAsyncRunCompletedDelegate, ENeuralNetworkDelegateThreadMode& InDelegateThreadMode,
+	FCriticalSection& InResoucesCriticalSection)
 #ifdef WITH_UE_AND_ORT_SUPPORT
 	: OnAsyncRunCompletedDelegate(InOutOnAsyncRunCompletedDelegate)
+	, DelegateThreadMode(InDelegateThreadMode)
 	, ResoucesCriticalSection(InResoucesCriticalSection)
 #endif
 {
@@ -332,7 +334,7 @@ bool UNeuralNetwork::FImplBackEndUEAndORT::IsGPUSupported()
 }
 
 bool UNeuralNetwork::FImplBackEndUEAndORT::Load(TSharedPtr<FImplBackEndUEAndORT>& InOutImplBackEndUEAndORT,
-	FOnAsyncRunCompleted& InOutOnAsyncRunCompletedDelegate, FCriticalSection& InOutResoucesCriticalSection,
+	FOnAsyncRunCompleted& InOutOnAsyncRunCompletedDelegate, ENeuralNetworkDelegateThreadMode& InOutDelegateThreadMode, FCriticalSection& InOutResoucesCriticalSection,
 	TArray<bool>& OutAreInputTensorSizesVariable, const TArray<uint8>& InModelReadFromFileInBytes, const FString& InModelFullFilePath,
 	const ENeuralDeviceType InDeviceType, const ENeuralDeviceType InInputDeviceType, const ENeuralDeviceType InOutputDeviceType)
 {
@@ -350,8 +352,8 @@ bool UNeuralNetwork::FImplBackEndUEAndORT::Load(TSharedPtr<FImplBackEndUEAndORT>
 		}
 
 		// Initialize and configure InOutImplBackEndUEAndORT
-		if (!UNeuralNetwork::FImplBackEndUEAndORT::InitializedAndConfigureMembers(InOutImplBackEndUEAndORT, InOutOnAsyncRunCompletedDelegate,
-				InOutResoucesCriticalSection, InModelFullFilePath, InDeviceType))
+		if (!UNeuralNetwork::FImplBackEndUEAndORT::InitializedAndConfigureMembers(InOutImplBackEndUEAndORT, InOutOnAsyncRunCompletedDelegate, InOutDelegateThreadMode,
+			InOutResoucesCriticalSection, InModelFullFilePath, InDeviceType))
 		{
 			UE_LOG(LogNeuralNetworkInference, Warning, TEXT("FImplBackEndUEAndORT::Load(): InitializedAndConfigureMembers failed."));
 			return false;
@@ -472,12 +474,13 @@ void UNeuralNetwork::FImplBackEndUEAndORT::EnsureAsyncTaskCompletion() const
 }
 
 bool UNeuralNetwork::FImplBackEndUEAndORT::InitializedAndConfigureMembers(TSharedPtr<FImplBackEndUEAndORT>& InOutImplBackEndUEAndORT,
-	FOnAsyncRunCompleted& InOutOnAsyncRunCompletedDelegate, FCriticalSection& InOutResoucesCriticalSection, const FString& InModelFullFilePath, const ENeuralDeviceType InDeviceType)
+	FOnAsyncRunCompleted& InOutOnAsyncRunCompletedDelegate, ENeuralNetworkDelegateThreadMode& InOutDelegateThreadMode, FCriticalSection& InOutResoucesCriticalSection,
+	const FString& InModelFullFilePath, const ENeuralDeviceType InDeviceType)
 {
 	// Initialize InOutImplBackEndUEAndORT
 	if (!InOutImplBackEndUEAndORT.IsValid())
 	{
-		InOutImplBackEndUEAndORT = MakeShared<FImplBackEndUEAndORT>(InOutOnAsyncRunCompletedDelegate, InOutResoucesCriticalSection);
+		InOutImplBackEndUEAndORT = MakeShared<FImplBackEndUEAndORT>(InOutOnAsyncRunCompletedDelegate, InOutDelegateThreadMode, InOutResoucesCriticalSection);
 
 		// Set up ORT and create an environment
 		const char* const ModelFullFilePathCharPtr = TCHAR_TO_ANSI(*InModelFullFilePath);
@@ -528,26 +531,20 @@ bool UNeuralNetwork::FImplBackEndUEAndORT::ConfigureMembers(const ENeuralDeviceT
 		FD3D12Adapter&				RhiAdapter = Rhi->GetAdapter(0);
 		const FD3D12AdapterDesc&	RhiAdapterDesc = RhiAdapter.GetDesc();
 
-		UE_LOG(LogNeuralNetworkInference, Display, TEXT("FImplBackEndUEAndORT::ConfigureMembers(): NNI available adapters %d"), Rhi->GetNumAdapters());
+		UE_LOG(LogNeuralNetworkInference, Display, TEXT("%d available RHI adapters. NNI using RHI adapter %s with LUID:%0x%0x."),
+			Rhi->GetNumAdapters(), RhiAdapterDesc.Desc.Description, RhiAdapterDesc.Desc.AdapterLuid.HighPart, RhiAdapterDesc.Desc.AdapterLuid.LowPart);
 
-		for (int32 ic = 0; ic < Rhi->GetNumAdapters(); ++ic)
+		if (Rhi->GetNumAdapters() > 1)
 		{
-			const FD3D12AdapterDesc& CurrAdapterDesc = Rhi->GetAdapter(ic).GetDesc();
-
-			UE_LOG(LogNeuralNetworkInference, Display,
-				TEXT("   Adapter [%d] Name:%s with LUID:%0x%0x"),
-				ic,
-				CurrAdapterDesc.Desc.Description,
-				CurrAdapterDesc.Desc.AdapterLuid.HighPart,
-				CurrAdapterDesc.Desc.AdapterLuid.LowPart);
+			UE_LOG(LogNeuralNetworkInference, Display, TEXT("All available RHI adapters:"));
+			for (int32 AdapterIndex = 0; AdapterIndex < Rhi->GetNumAdapters(); ++AdapterIndex)
+			{
+				const FD3D12AdapterDesc& CurrAdapterDesc = Rhi->GetAdapter(AdapterIndex).GetDesc();
+				UE_LOG(LogNeuralNetworkInference, Display, TEXT("  - Adapter [%d] Name:%s with LUID:%0x%0x."),
+					AdapterIndex, CurrAdapterDesc.Desc.Description, CurrAdapterDesc.Desc.AdapterLuid.HighPart, CurrAdapterDesc.Desc.AdapterLuid.LowPart);
+			}
 		}
 
-		UE_LOG(LogNeuralNetworkInference, Display,
-			TEXT("FImplBackEndUEAndORT::ConfigureMembers(): NNI using adapter %s with LUID:%0x%0x"),
-			RhiAdapterDesc.Desc.Description,
-			RhiAdapterDesc.Desc.AdapterLuid.HighPart,
-			RhiAdapterDesc.Desc.AdapterLuid.LowPart);
-
 		if (Rhi->GetNumAdapters() > 1 || RhiAdapterDesc.NumDeviceNodes > 1)
 		{
 			UE_LOG(LogNeuralNetworkInference, Warning,
@@ -945,7 +942,31 @@ void UNeuralNetwork::FImplBackEndUEAndORT::RunSessionAsync(const ENeuralDeviceTy
 
 	RunSessionImpl(InDeviceType, InInputDeviceType, InOutputDeviceType);
 
-	OnAsyncRunCompletedDelegate.ExecuteIfBound();
+	// Execute OnAsyncRunCompletedDelegate on main thread
+	if (DelegateThreadMode == ENeuralNetworkDelegateThreadMode::GameThread)
+	{
+		if (OnAsyncRunCompletedDelegate.IsBound())
+		{
+			std::atomic<bool> IsComputeFinished(false);
+			AsyncTask(ENamedThreads::GameThread, [this, &IsComputeFinished]
+			{
+				OnAsyncRunCompletedDelegate.ExecuteIfBound();
+				IsComputeFinished = true;
+			});
+			while (!IsComputeFinished)
+			{
+				FPlatformProcess::Sleep(0.1e-3);
+			}
+		}
+	}
+	else if (DelegateThreadMode == ENeuralNetworkDelegateThreadMode::AnyThread)
+	{
+		OnAsyncRunCompletedDelegate.ExecuteIfBound();
+	}
+	else
+	{
+		UE_LOG(LogNeuralNetworkInference, Warning, TEXT("FImplBackEndUEAndORT::RunSessionAsync(): Unknown DelegateThreadMode = %d."), (int32)DelegateThreadMode);
+	}
 }
 
 void UNeuralNetwork::FImplBackEndUEAndORT::RunSessionSync(const ENeuralDeviceType InDeviceType, const ENeuralDeviceType InInputDeviceType,
diff --git a/Engine/Plugins/Experimental/NeuralNetworkInference/Source/NeuralNetworkInference/Private/NeuralNetworkImplBackEndUEAndORT.h b/Engine/Plugins/Experimental/NeuralNetworkInference/Source/NeuralNetworkInference/Private/NeuralNetworkImplBackEndUEAndORT.h
--- a/Engine/Plugins/Experimental/NeuralNetworkInference/Source/NeuralNetworkInference/Private/NeuralNetworkImplBackEndUEAndORT.h
+++ b/Engine/Plugins/Experimental/NeuralNetworkInference/Source/NeuralNetworkInference/Private/NeuralNetworkImplBackEndUEAndORT.h
@@ -43,7 +43,7 @@ public:
 	TArray<FNeuralTensor> InputTensors;
 	TArray<FNeuralTensor> OutputTensors;
 
-	FImplBackEndUEAndORT(FOnAsyncRunCompleted& InOutOnAsyncRunCompletedDelegate, FCriticalSection& InResoucesCriticalSection);
+	FImplBackEndUEAndORT(FOnAsyncRunCompleted& InOutOnAsyncRunCompletedDelegate, ENeuralNetworkDelegateThreadMode& InDelegateThreadMode, FCriticalSection& InResoucesCriticalSection);
 
 	~FImplBackEndUEAndORT();
 
@@ -52,8 +52,9 @@ public:
 	static bool IsGPUSupported();
 
 	static bool Load(TSharedPtr<FImplBackEndUEAndORT>& InOutImplBackEndUEAndORT, FOnAsyncRunCompleted& InOutOnAsyncRunCompletedDelegate,
-		FCriticalSection& InOutResoucesCriticalSection, TArray<bool>& OutAreInputTensorSizesVariable, const TArray<uint8>& InModelReadFromFileInBytes,
-		const FString& InModelFullFilePath, const ENeuralDeviceType InDeviceType, const ENeuralDeviceType InInputDeviceType, const ENeuralDeviceType InOutputDeviceType);
+		ENeuralNetworkDelegateThreadMode& InOutDelegateThreadMode, FCriticalSection& InOutResoucesCriticalSection, TArray<bool>& OutAreInputTensorSizesVariable,
+		const TArray<uint8>& InModelReadFromFileInBytes, const FString& InModelFullFilePath, const ENeuralDeviceType InDeviceType, const ENeuralDeviceType InInputDeviceType,
+		const ENeuralDeviceType InOutputDeviceType);
 
 	void Run(const ENeuralNetworkSynchronousMode InSynchronousMode, const ENeuralDeviceType InDeviceType, const ENeuralDeviceType InInputDeviceType,
 		const ENeuralDeviceType InOutputDeviceType);
@@ -62,6 +63,7 @@ public:
 private:
 	/** Async support */
 	FOnAsyncRunCompleted& OnAsyncRunCompletedDelegate;
+	ENeuralNetworkDelegateThreadMode& DelegateThreadMode;
 	FCriticalSection& ResoucesCriticalSection;
 	/** Network-related variables */
 	TUniquePtr<Ort::Env> Environment;
@@ -108,7 +110,7 @@ private:
 	void EnsureAsyncTaskCompletion() const;
 
 	static bool InitializedAndConfigureMembers(TSharedPtr<FImplBackEndUEAndORT>& InOutImplBackEndUEAndORT, FOnAsyncRunCompleted& InOutOnAsyncRunCompletedDelegate,
-		FCriticalSection& InOutResoucesCriticalSection, const FString& InModelFullFilePath, const ENeuralDeviceType InDeviceType);
+		ENeuralNetworkDelegateThreadMode& InOutDelegateThreadMode, FCriticalSection& InOutResoucesCriticalSection, const FString& InModelFullFilePath, const ENeuralDeviceType InDeviceType);
 
 	bool ConfigureMembers(const ENeuralDeviceType InDeviceType);
 
diff --git a/Engine/Plugins/Experimental/NeuralNetworkInference/Source/NeuralNetworkInference/Public/NeuralNetwork.h b/Engine/Plugins/Experimental/NeuralNetworkInference/Source/NeuralNetworkInference/Public/NeuralNetwork.h
--- a/Engine/Plugins/Experimental/NeuralNetworkInference/Source/NeuralNetworkInference/Public/NeuralNetwork.h
+++ b/Engine/Plugins/Experimental/NeuralNetworkInference/Source/NeuralNetworkInference/Public/NeuralNetwork.h
@@ -26,6 +26,19 @@ enum class ENeuralNetworkSynchronousMode : uint8
 	Asynchronous
 };
 
+UENUM()
+enum class ENeuralNetworkDelegateThreadMode : uint8
+{
+	GameThread, /* Recommended and default value. The UNeuralNetwork delegate will be called from the game thread. */
+	/**
+	 * Not recommended, use at your own risk.
+	 * The UNeuralNetwork delegate could be called from any thread.
+	 * Running UClass functions from background threads is not safe (e.g., it might crash if the editor is closed while accessing UNeuralNetwork information).
+	 * Thus "AnyThread" is only safe if you have guarantees that the program will not be terminated while calling UNeuralNetwork functions.
+	 */
+	AnyThread
+};
+
 /**
  * UNeuralNetwork is UE's representation for deep learning and neural network models. It supports the industry standard ONNX model format.
  * All major frameworks (PyTorch, TensorFlow, MXNet, Caffe2, etc.) provide converters to ONNX.
@@ -116,7 +129,6 @@ public:
 	ENeuralNetworkSynchronousMode GetSynchronousMode() const;
 	void SetSynchronousMode(const ENeuralNetworkSynchronousMode InSynchronousMode);
 
-
 	/**
 	 * GetOnAsyncRunCompletedDelegate() returns a FOnAsyncRunCompleted delegate that will be called when async UNeuralNetwork::Run() is completed.
 	 * This FOnAsyncRunCompleted delegate could be triggered from any thread but will only be triggered if SynchronousMode == ENeuralNetworkSynchronousMode::Asynchronous.
@@ -124,6 +136,8 @@ public:
 	 */
 	DECLARE_DELEGATE(FOnAsyncRunCompleted);
 	FOnAsyncRunCompleted& GetOnAsyncRunCompletedDelegate();
+	ENeuralNetworkDelegateThreadMode GetOnAsyncRunCompletedDelegateMode() const;
+	void SetOnAsyncRunCompletedDelegateMode(const ENeuralNetworkDelegateThreadMode InDelegateThreadMode);
 
 	/**
 	 * Getter and setter functions for BackEnd.
@@ -218,11 +232,15 @@ protected:
 	ENeuralDeviceType OutputDeviceType;
 	
 	/**
-	 * Whether UNeuralNetwork::Run() will block the thread until completed (Synchronous), or whether it will run on a background thread, not blocking the calling thread (Asynchronous).
-	 * @see ENeuralNetworkSynchronousMode for more details.
+	 * SynchronousMode defines whether UNeuralNetwork::Run() will block the thread until completed (Synchronous), or whether it will run on a background thread, not
+	 * blocking the calling thread (Asynchronous).
+	 * If asynchronous, DelegateThreadMode will define whether the callback delegate is called from the game thread (highly recommended) or from any available thread (not fully thread safe).
+	 * @see ENeuralNetworkSynchronousMode, ENeuralNetworkDelegateThreadMode for more details.
 	 */
-	UPROPERTY(VisibleAnywhere, Category = "Neural Network Inference")
+	UPROPERTY(Transient, VisibleAnywhere, Category = "Neural Network Inference")
 	ENeuralNetworkSynchronousMode SynchronousMode;
+	UPROPERTY(Transient, VisibleAnywhere, Category = "Neural Network Inference")
+	ENeuralNetworkDelegateThreadMode DelegateThreadMode;
 	
 	/**
 	 * @see ENeuralBackEnd for more details.
@@ -237,6 +255,7 @@ protected:
 	FString ModelFullFilePath;
 
 private:
+	UPROPERTY(Transient, VisibleAnywhere, Category = "Neural Network Inference")
 	bool bIsLoaded;
 
 	/**
@@ -249,6 +268,7 @@ private:
 	TArray<uint8> ModelReadFromFileInBytes;
 
 	/** Whether some of the FNeuralTensor of InputTensor have flexible/variable dimensions. */
+	UPROPERTY(Transient, VisibleAnywhere, Category = "Neural Network Inference")
 	TArray<bool> AreInputTensorSizesVariable;
 
 	/**
@@ -256,6 +276,7 @@ private:
 	 * Otherwise, BackEndForCurrentPlatform will be set to the optimal BackEnd given the current platform.
 	 * @see ENeuralBackEnd for more details.
 	 */
+	UPROPERTY(Transient, VisibleAnywhere, Category = "Neural Network Inference")
 	ENeuralBackEnd BackEndForCurrentPlatform;
 
 	/**
diff --git a/Engine/Plugins/Experimental/NeuralNetworkInference/Source/NeuralNetworkInferenceQA/Private/ModelUnitTester.cpp b/Engine/Plugins/Experimental/NeuralNetworkInference/Source/NeuralNetworkInferenceQA/Private/ModelUnitTester.cpp
--- a/Engine/Plugins/Experimental/NeuralNetworkInference/Source/NeuralNetworkInferenceQA/Private/ModelUnitTester.cpp
+++ b/Engine/Plugins/Experimental/NeuralNetworkInference/Source/NeuralNetworkInferenceQA/Private/ModelUnitTester.cpp
@@ -321,10 +321,12 @@ bool FModelUnitTester::ModelAccuracyTest(UNeuralNetwork* InOutNetwork, const ENe
 	const ENeuralDeviceType OriginalInputDeviceType = InOutNetwork->GetInputDeviceType();
 	const ENeuralDeviceType OriginalOutputDeviceType = InOutNetwork->GetOutputDeviceType();
 	const ENeuralNetworkSynchronousMode OriginalSynchronousMode = InOutNetwork->GetSynchronousMode();
+	const ENeuralNetworkDelegateThreadMode OriginalDelegateThreadMode = InOutNetwork->GetOnAsyncRunCompletedDelegateMode();
 	const ENeuralBackEnd OriginalBackEnd = InOutNetwork->GetBackEnd();
 	
 	// Set (a)synchronous Mode
 	InOutNetwork->SetSynchronousMode(InSynchronousMode);
+	InOutNetwork->SetOnAsyncRunCompletedDelegateMode(ENeuralNetworkDelegateThreadMode::AnyThread);
 
 	// Set back end
 	if (!InOutNetwork->SetBackEnd(InBackEnd))
@@ -412,6 +414,7 @@ bool FModelUnitTester::ModelAccuracyTest(UNeuralNetwork* InOutNetwork, const ENe
 	// Reset to original network state
 	InOutNetwork->SetDeviceType(/*DeviceType*/OriginalDeviceType, /*InputDeviceType*/OriginalInputDeviceType, /*OutputDeviceType*/OriginalOutputDeviceType);
 	InOutNetwork->SetSynchronousMode(OriginalSynchronousMode);
+	InOutNetwork->SetOnAsyncRunCompletedDelegateMode(OriginalDelegateThreadMode);
 	InOutNetwork->SetBackEnd(OriginalBackEnd);
 	// Test successful
 	return bDidGlobalTestPassed;
