Adding NNI documentation - Part 1:
- NNI/Public/NeuralNetwork.h cleaned a bit
- All Enum classes moved to enumclasses.h (not renamed yet).

#rnx

diff --git a/Engine/Plugins/Experimental/MLDeformer/Source/MLDeformer/Private/MLDeformerAsset.cpp b/Engine/Plugins/Experimental/MLDeformer/Source/MLDeformer/Private/MLDeformerAsset.cpp
--- a/Engine/Plugins/Experimental/MLDeformer/Source/MLDeformer/Private/MLDeformerAsset.cpp
+++ b/Engine/Plugins/Experimental/MLDeformer/Source/MLDeformer/Private/MLDeformerAsset.cpp
@@ -65,7 +65,6 @@ void UMLDeformerAsset::PostLoad()
 	Super::PostLoad();
 	if (NeuralNetwork)
 	{
-		NeuralNetwork->SetBackEnd(ENeuralBackEnd::Auto);
 		NeuralNetwork->SetDeviceType(/*DeviceType*/ENeuralDeviceType::GPU, /*InputDeviceType*/ENeuralDeviceType::CPU, /*OutputDeviceType*/ENeuralDeviceType::GPU);
 	}
 
diff --git a/Engine/Plugins/Experimental/MLDeformer/Source/MLDeformerEditor/Private/MLDeformerEditorToolkit.cpp b/Engine/Plugins/Experimental/MLDeformer/Source/MLDeformerEditor/Private/MLDeformerEditorToolkit.cpp
--- a/Engine/Plugins/Experimental/MLDeformer/Source/MLDeformerEditor/Private/MLDeformerEditorToolkit.cpp
+++ b/Engine/Plugins/Experimental/MLDeformer/Source/MLDeformerEditor/Private/MLDeformerEditorToolkit.cpp
@@ -374,7 +374,6 @@ bool FMLDeformerEditorToolkit::TryLoadOnnxFile() const
 		if (Network->Load(OnnxFile))
 		{
 			EditorData->GetDeformerAsset()->NeuralNetwork = Network;
-			EditorData->GetDeformerAsset()->NeuralNetwork->SetBackEnd(ENeuralBackEnd::Auto);
 			EditorData->GetDeformerAsset()->NeuralNetwork->SetDeviceType(/*DeviceType*/ENeuralDeviceType::GPU, /*InputDeviceType*/ENeuralDeviceType::CPU, /*OutputDeviceType*/ENeuralDeviceType::GPU);
 
 			// Recreate the data providers.
diff --git a/Engine/Plugins/Experimental/NeuralNetworkInference/Source/NeuralNetworkInference/Private/NeuralNetwork.cpp b/Engine/Plugins/Experimental/NeuralNetworkInference/Source/NeuralNetworkInference/Private/NeuralNetwork.cpp
--- a/Engine/Plugins/Experimental/NeuralNetworkInference/Source/NeuralNetworkInference/Private/NeuralNetwork.cpp
+++ b/Engine/Plugins/Experimental/NeuralNetworkInference/Source/NeuralNetworkInference/Private/NeuralNetwork.cpp
@@ -21,18 +21,18 @@
 struct FPrivateNeuralNetwork
 {
 public:
-	static ENeuralBackEnd SetBackEndForCurrentPlatform(const ENeuralBackEnd InBackEnd);
+	static UNeuralNetwork::ENeuralBackEnd SetBackEndForCurrentPlatform(const UNeuralNetwork::ENeuralBackEnd InBackEnd);
 };
 
-ENeuralBackEnd FPrivateNeuralNetwork::SetBackEndForCurrentPlatform(const ENeuralBackEnd InBackEnd)
+UNeuralNetwork::ENeuralBackEnd FPrivateNeuralNetwork::SetBackEndForCurrentPlatform(const UNeuralNetwork::ENeuralBackEnd InBackEnd)
 {
 	// Auto
-	if (InBackEnd == ENeuralBackEnd::Auto)
+	if (InBackEnd == UNeuralNetwork::ENeuralBackEnd::Auto)
 	{
 #ifdef WITH_UE_AND_ORT_SUPPORT
-		return ENeuralBackEnd::UEAndORT;
+		return UNeuralNetwork::ENeuralBackEnd::UEAndORT;
 #else //WITH_UE_AND_ORT_SUPPORT
-		return ENeuralBackEnd::UEOnly;
+		return UNeuralNetwork::ENeuralBackEnd::UEOnly;
 #endif //WITH_UE_AND_ORT_SUPPORT
 	}
 	// Otherwise
@@ -50,8 +50,8 @@ UNeuralNetwork::UNeuralNetwork()
 	, OutputDeviceType(ENeuralDeviceType::CPU)
 	, SynchronousMode(ENeuralNetworkSynchronousMode::Synchronous)
 	, DelegateThreadMode(ENeuralNetworkDelegateThreadMode::GameThread)
-	, BackEnd(ENeuralBackEnd::Auto)
 	, bIsLoaded(false)
+	, BackEnd(ENeuralBackEnd::Auto)
 	, BackEndForCurrentPlatform(FPrivateNeuralNetwork::SetBackEndForCurrentPlatform(BackEnd))
 {
 }
@@ -187,33 +187,6 @@ void UNeuralNetwork::SetOnAsyncRunCompletedDelegateMode(const ENeuralNetworkDele
 	DelegateThreadMode = InDelegateThreadMode;
 }
 
-ENeuralBackEnd UNeuralNetwork::GetBackEnd() const
-{
-	return BackEnd;
-}
-
-ENeuralBackEnd UNeuralNetwork::GetBackEndForCurrentPlatform() const
-{
-	return BackEndForCurrentPlatform;
-}
-
-bool UNeuralNetwork::SetBackEnd(const ENeuralBackEnd InBackEnd)
-{
-	BackEnd = InBackEnd;
-	const ENeuralBackEnd NewBackEndForCurrentPlatform = FPrivateNeuralNetwork::SetBackEndForCurrentPlatform(BackEnd);
-	// Reload only required if BackEndForCurrentPlatform changes (regardless of whether BackEnd changed).
-	// BackEndForCurrentPlatform does not necessarily change if BackEnd changes. E.g., changing from UEAndORT into Auto in Windows will result in BackEndForCurrentPlatform = UEAndORT in both cases.
-	if (BackEndForCurrentPlatform != NewBackEndForCurrentPlatform)
-	{
-		BackEndForCurrentPlatform = NewBackEndForCurrentPlatform;
-		if (bIsLoaded) // No need to re-load if not bIsLoaded
-		{
-			Load();
-		}
-	}
-	return IsLoaded();
-}
-
 void UNeuralNetwork::ResetStats()
 {
 	ComputeStatsModule.ResetStats();
@@ -616,6 +589,33 @@ bool UNeuralNetwork::Load(TArray<FNeuralTensor>& InTensors, const TArray<FNeural
 	return IsLoaded();
 }
 
+UNeuralNetwork::ENeuralBackEnd UNeuralNetwork::GetBackEnd() const
+{
+	return BackEnd;
+}
+
+UNeuralNetwork::ENeuralBackEnd UNeuralNetwork::GetBackEndForCurrentPlatform() const
+{
+	return BackEndForCurrentPlatform;
+}
+
+bool UNeuralNetwork::SetBackEnd(const UNeuralNetwork::ENeuralBackEnd InBackEnd)
+{
+	BackEnd = InBackEnd;
+	const ENeuralBackEnd NewBackEndForCurrentPlatform = FPrivateNeuralNetwork::SetBackEndForCurrentPlatform(BackEnd);
+	// Reload only required if BackEndForCurrentPlatform changes (regardless of whether BackEnd changed).
+	// BackEndForCurrentPlatform does not necessarily change if BackEnd changes. E.g., changing from UEAndORT into Auto in Windows will result in BackEndForCurrentPlatform = UEAndORT in both cases.
+	if (BackEndForCurrentPlatform != NewBackEndForCurrentPlatform)
+	{
+		BackEndForCurrentPlatform = NewBackEndForCurrentPlatform;
+		if (bIsLoaded) // No need to re-load if not bIsLoaded
+		{
+			Load();
+		}
+	}
+	return IsLoaded();
+}
+
 #if WITH_EDITOR
 UAssetImportData* UNeuralNetwork::GetAssetImportData() const
 {
diff --git a/Engine/Plugins/Experimental/NeuralNetworkInference/Source/NeuralNetworkInference/Public/NeuralEnumClasses.h b/Engine/Plugins/Experimental/NeuralNetworkInference/Source/NeuralNetworkInference/Public/NeuralEnumClasses.h
--- a/Engine/Plugins/Experimental/NeuralNetworkInference/Source/NeuralNetworkInference/Public/NeuralEnumClasses.h
+++ b/Engine/Plugins/Experimental/NeuralNetworkInference/Source/NeuralNetworkInference/Public/NeuralEnumClasses.h
@@ -7,18 +7,6 @@
 #include <type_traits>
 #include "NeuralEnumClasses.generated.h"
 
-/**
- * Whether UNeuralNetwork will use the highly optimized UnrealEngine-and-ONNXRuntime-based back end (UEAndORT) or the less optimized but fully cross platform only-UE one (UEOnly).
- * We recommend using Auto, which will find and use the optimal back end for each platform.
- */
-UENUM()
-enum class ENeuralBackEnd : uint8
-{
-	UEAndORT, /* UnrealEngine-and-ONNXRuntime-accelerated back end, ideal for those platforms that support it. WITH_UE_AND_ORT_SUPPORT is the C++ define macro that checks whether UEAndORT support exists for the current platform.*/
-	UEOnly, /* It might be slower than the UEAndORT back end, but it will compile in all platforms and OSs compatible with Unreal Engine. */
-	Auto /* Recommended setting. It will use UEAndORT if possible, and fall back to UEOnly only for those platforms that do not support UEAndORT yet. */
-};
-
 UENUM()
 enum class ENeuralDataType : uint8
 {
@@ -42,6 +30,55 @@ enum class ENeuralDeviceType : uint8
 	None
 };
 
+/**
+ * Whether UNeuralNetwork::Run() will block the thread until completed (Synchronous), or whether it will run on a background thread,
+ * not blocking the calling thread (Asynchronous).
+ */
+UENUM()
+enum class ENeuralNetworkSynchronousMode : uint8
+{
+	Synchronous, /* UNeuralNetwork::Run() will block the thread until the network evaluation (i.e., forward pass) has finished. */
+	/**
+	 * UNeuralNetwork::Run() will initialize a forward pass request on a background thread, not blocking the thread that called it.
+	 * The user should register to UNeuralNetwork's delegate to know when the forward pass has finished.
+	 *
+	 * Very important: It takes ~1 millisecond to start the background thread. If your network runs synchronously faster than 1 msec,
+	 * using asynchronous running will make the game (main) thread slower than running it synchronously.
+	 */
+	Asynchronous
+};
+
+UENUM()
+enum class ENeuralNetworkDelegateThreadMode : uint8
+{
+	GameThread, /* Recommended and default value. The UNeuralNetwork delegate will be called from the game thread. */
+	/**
+	 * Not recommended, use at your own risk.
+	 * The UNeuralNetwork delegate could be called from any thread.
+	 * Running UClass functions from background threads is not safe (e.g., it might crash if the editor is closed while accessing UNeuralNetwork information).
+	 * Thus "AnyThread" is only safe if you have guarantees that the program will not be terminated while calling UNeuralNetwork functions.
+	 */
+	AnyThread
+};
+
+
+/**
+ * Although conceptually this could apply to both the CPU and GPU versions, in practice only the GPU performance is affected by this setting.
+ * Input and Intermediate(Not)Initialized currently share the same attributes because input might become intermediate (e.g., if input tensor fed into a ReLU, which simply modifies
+ * the input FNeuralTensor). However, Intermediate(Not)Initialized and Output do not copy the memory from CPU to GPU but rather simply allocates it.
+ * Output might also become Intermediate(Not)Initialized (e.g., if Output -> ReLU -> Output), so it is kept as ReadWrite rather than written once to account for this.
+ */
+UENUM()
+enum class ENeuralTensorTypeGPU : uint8
+{
+	Generic,					/** Generic tensor that works in every situation (ReadWrite), although it might not be the most efficient one. */
+	Input,						/** Input tensor of the UNeuralNetworkLegacy. Copied from CPU and ReadWrite (but usually ReadOnly). */
+	IntermediateNotInitialized,	/** Intermediate tensor of the UNeuralNetworkLegacy (output of at least a layer and input of at least some other layer). Not copied from CPU, ReadWrite, and transient. */
+	IntermediateInitialized,	/** Intermediate tensor that is initialized with CPU data (e.g., XWithZeros in FConvTranpose). Copied from CPU. */
+	Output,						/** Output tensor of the UNeuralNetworkLegacy. Not copied from CPU and ReadWrite. */
+	Weight						/** Weights of a particular operator/layer. Copied from CPU, ReadOnly, and initialized from CPU memory. */
+};
+
 
 
 /**
diff --git a/Engine/Plugins/Experimental/NeuralNetworkInference/Source/NeuralNetworkInference/Public/NeuralNetwork.h b/Engine/Plugins/Experimental/NeuralNetworkInference/Source/NeuralNetworkInference/Public/NeuralNetwork.h
--- a/Engine/Plugins/Experimental/NeuralNetworkInference/Source/NeuralNetworkInference/Public/NeuralNetwork.h
+++ b/Engine/Plugins/Experimental/NeuralNetworkInference/Source/NeuralNetworkInference/Public/NeuralNetwork.h
@@ -8,37 +8,6 @@
 #include "NeuralStats.h"
 #include "NeuralNetwork.generated.h"
 
-/**
- * Whether UNeuralNetwork::Run() will block the thread until completed (Synchronous), or whether it will run on a background thread,
- * not blocking the calling thread (Asynchronous).
- */
-UENUM()
-enum class ENeuralNetworkSynchronousMode : uint8
-{
-	Synchronous, /* UNeuralNetwork::Run() will block the thread until the network evaluation (i.e., forward pass) has finished. */
-	/**
-	 * UNeuralNetwork::Run() will initialize a forward pass request on a background thread, not blocking the thread that called it.
-	 * The user should register to UNeuralNetwork's delegate to know when the forward pass has finished.
-	 *
-	 * Very important: It takes ~1 millisecond to start the background thread. If your network runs synchronously faster than 1 msec,
-	 * using asynchronous running will make the game (main) thread slower than running it synchronously.
-	 */
-	Asynchronous
-};
-
-UENUM()
-enum class ENeuralNetworkDelegateThreadMode : uint8
-{
-	GameThread, /* Recommended and default value. The UNeuralNetwork delegate will be called from the game thread. */
-	/**
-	 * Not recommended, use at your own risk.
-	 * The UNeuralNetwork delegate could be called from any thread.
-	 * Running UClass functions from background threads is not safe (e.g., it might crash if the editor is closed while accessing UNeuralNetwork information).
-	 * Thus "AnyThread" is only safe if you have guarantees that the program will not be terminated while calling UNeuralNetwork functions.
-	 */
-	AnyThread
-};
-
 /**
  * UNeuralNetwork is UE's representation for deep learning and neural network models. It supports the industry standard ONNX model format.
  * All major frameworks (PyTorch, TensorFlow, MXNet, Caffe2, etc.) provide converters to ONNX.
@@ -49,7 +18,7 @@ enum class ENeuralNetworkDelegateThreadMode : uint8
  *		UNeuralNetwork* Network = NewObject<UNeuralNetwork>((UObject*)GetTransientPackage(), UNeuralNetwork::StaticClass());
  *		// Try to load the network and set the device (CPU/GPU)
  *		const FString ONNXModelFilePath = TEXT("SOME_PARENT_FOLDER/SOME_ONNX_FILE_NAME.onnx");
- *		if (Network->Load(ONNXModelFilePath)) 
+ *		if (Network->Load(ONNXModelFilePath))
  *		{
  *			Network->SetDeviceType(ENeuralDeviceType::CPU); // Set to CPU/GPU mode
  *		}
@@ -80,7 +49,8 @@ enum class ENeuralNetworkDelegateThreadMode : uint8
  *			InputDataPointer[Index] = ...;
  *
  * 3.3. Alternative - Networks with multiple input/output tensors:
- * - Multiple inputs: Add InTensorIndex to GetInputTensor(InTensorIndex) or GetInputDataPointerMutable(InTensorIndex) in the examples above, or use GetInputTensors() instead.
+ * - Multiple inputs: Add InTensorIndex to GetInputTensor(InTensorIndex) or GetInputDataPointerMutable(InTensorIndex) in the examples above, or use
+ *   GetInputTensors() instead.
  * - Multiple outputs: Add InTensorIndex to GetOutputTensor(InTensorIndex) in the examples above or use GetOutputTensors() instead.
  */
 UCLASS(BlueprintType)
@@ -120,7 +90,8 @@ public:
 	ENeuralDeviceType GetDeviceType() const;
 	ENeuralDeviceType GetInputDeviceType() const;
 	ENeuralDeviceType GetOutputDeviceType() const;
-	void SetDeviceType(const ENeuralDeviceType InDeviceType, const ENeuralDeviceType InInputDeviceType = ENeuralDeviceType::CPU, const ENeuralDeviceType InOutputDeviceType = ENeuralDeviceType::CPU);
+	void SetDeviceType(const ENeuralDeviceType InDeviceType, const ENeuralDeviceType InInputDeviceType = ENeuralDeviceType::CPU,
+		const ENeuralDeviceType InOutputDeviceType = ENeuralDeviceType::CPU);
 
 	/**
 	 * Getter and setter functions for SynchronousMode.
@@ -131,37 +102,28 @@ public:
 
 	/**
 	 * GetOnAsyncRunCompletedDelegate() returns a FOnAsyncRunCompleted delegate that will be called when async UNeuralNetwork::Run() is completed.
-	 * This FOnAsyncRunCompleted delegate could be triggered from any thread but will only be triggered if SynchronousMode == ENeuralNetworkSynchronousMode::Asynchronous.
-	 * If SynchronousMode == ENeuralNetworkSynchronousMode::Synchronous, UNeuralNetwork::Run() will block the calling thread until completed, so a callback delegate is not required.
+	 * - If SynchronousMode == ENeuralNetworkSynchronousMode::Asynchronous, the FOnAsyncRunCompleted delegate could be triggered from any thread.
+	 * - If SynchronousMode == ENeuralNetworkSynchronousMode::Synchronous, UNeuralNetwork::Run() will block the calling thread until completed, so a
+	 *   callback delegate is not required.
 	 */
 	DECLARE_DELEGATE(FOnAsyncRunCompleted);
 	FOnAsyncRunCompleted& GetOnAsyncRunCompletedDelegate();
 	ENeuralNetworkDelegateThreadMode GetOnAsyncRunCompletedDelegateMode() const;
 	void SetOnAsyncRunCompletedDelegateMode(const ENeuralNetworkDelegateThreadMode InDelegateThreadMode);
-
-	/**
-	 * Getter and setter functions for BackEnd.
-	 * GetBackEnd()/GetBackEndForCurrentPlatform():
-	 * - If BackEnd == Auto, GetBackEnd() will return Auto and GetBackEndForCurrentPlatform() will return the actual BackEnd being used for the current platform (UEAndORT or UEOnly).
-	 * - If BackEnd != Auto, GetBackEnd() and GetBackEndForCurrentPlatform() will both return the same value (UEAndORT or UEOnly).
-	 * SetBackEnd() will modify both BackEnd and BackEndForCurrentPlatform and return IsLoaded(). 
-	 * Important! SetBackEnd is NOT THREAD SAFE, please ensure that no other Neural Network function such as Run() is running when SetBackEnd is called. 
-	 * @see ENeuralBackEnd for more details.
-	 */
-	ENeuralBackEnd GetBackEnd() const;
-	ENeuralBackEnd GetBackEndForCurrentPlatform() const;
-	bool SetBackEnd(const ENeuralBackEnd InBackEnd);
 	
 	/**
-	 * IsGPUCompatible will always return true for ENeuralBackEnd::UEOnly. For ENeuralBackEnd::UEAndORT, it will return:
-	 * - True if DX12 is enabled, meaning UEAndORT can run on both the CPU and GPU. Also true if the current platform is not Windows.
-	 * - False if DX12 is disabled, meaning UEAndORT can run only run on the CPU. The user will need to enable DX12 to be able to run GPU, switch to CPU, or switch to the UEOnly back end.
+	 * Whether GPU execution is supported for this platform. It will return:
+	 * - True if DX12 is enabled, meaning UEAndORT can run on both the CPU and GPU. Also true if the current platform is not Windows. I will also
+	 *   return true if the back end is UEOnly.
+	 * - False if DX12 is disabled, meaning UEAndORT can run only run on the CPU. The user will need to enable DX12 to be able to run GPU, switch to
+	 *   CPU, or switch to the UEOnly back end.
 	 */
 	bool IsGPUSupported() const;
 
 	/**
 	 * Functions to get/fill input.
-	 * These functions either take a TArray as input, return a modificable void* to fill the data, or return a constant FNeuralTensor(s) to see input properties (e.g., size or dimensions).
+	 * These functions either take a TArray as input, return a modificable void* to fill the data, or return a constant FNeuralTensor(s) to see input
+	 * properties (e.g., size or dimensions).
 	 */
 	const FNeuralTensor& GetInputTensor(const int32 InTensorIndex = 0) const;
 	void SetInputFromArrayCopy(const TArray<float>& InArray, const int32 InTensorIndex = 0);
@@ -176,35 +138,38 @@ public:
 	TArray<FNeuralTensor> CreateOutputArrayCopy() const;
 
 	/**
-	 * Functions to get output. The returned FNeuralTensor(s) are constant to prevent the user from modifying the tensor properties (e.g., size or dimensions).
+	 * Functions to get output. The returned FNeuralTensor(s) are constant to prevent the user from modifying the tensor properties (e.g., size or
+	 * dimensions).
 	 */
 	const FNeuralTensor& GetOutputTensor(const int32 InTensorIndex = 0) const;
 	int64 GetOutputTensorNumber() const;
 
 	/**
 	 * Non-efficient functions meant to be used only for debugging purposes.
-	 * - InputTensorsToCPU will send the CPU memory of the desired input tensor(s) to GPU memory. Used to debug InputDeviceType == ENeuralDeviceType::GPU.
-	 * - OutputTensorsToCPU will send the GPU memory of the desired output tensor(s) back to CPU memory. Used to debug OutputDeviceType == ENeuralDeviceType::GPU.
+	 * - InputTensorsToCPU copies the CPU memory of the desired input tensor(s) to GPU (to debug InputDeviceType == ENeuralDeviceType::GPU).
+	 * - OutputTensorsToCPU copies the GPU memory of the desired output tensor(s) back to CPU (to debug OutputDeviceType == ENeuralDeviceType::GPU).
 	 * @param InTensorIndexes If empty (default value), it will apply to all output tensors.
 	 */
 	void InputTensorsToGPU(const TArray<int32>& InTensorIndexes = TArray<int32>());
 	void OutputTensorsToCPU(const TArray<int32>& InTensorIndexes = TArray<int32>());
 
 	/**
-	 * Run() executes the forward pass on the current UNeuralNetwork given the current input FDeprecatedNeuralTensor(s), which were previously filled with
-	 * SetInputFromArrayCopy() or GetInputDataPointerMutable().
+	 * Run() executes the forward pass on the current UNeuralNetwork given the current input FDeprecatedNeuralTensor(s), which were previously filled
+	 * with SetInputFromArrayCopy() or GetInputDataPointerMutable().
 	 * Its output results can be retrieved with GetOutputTensor() or GetOutputTensors().
 	 *
-	 * If Run() is called asynchronously, this does not guarantee that calling SetInputFromArrayCopy multiple times will result in each one being applied for a different Run. The user is
-	 * responable of not calling SetInputFromArrayCopy until Run() is completed and its delegate (OnAsyncRunCompletedDelegate) called. Otherwise, the wrong results might be returned.
+	 * If Run() is called asynchronously, this does not guarantee that calling SetInputFromArrayCopy multiple times will result in each one being
+	 * applied for a different Run. The user is responable of not calling SetInputFromArrayCopy until Run() is completed and its delegate
+	 * (OnAsyncRunCompletedDelegate) called. Otherwise, the wrong results might be returned.
 	 */
 	void Run();
 
 	/**
-	 * Stats functions: 
+	 * Stats functions:
 	 * - GetLastInferenceTime will provide the last inference time measured milliseconds
-	 * - GetInferenceStats, returns Inference time statistics. (NumberSamples, Average, StdDev, Min, Max statistics measured in milliseconds)
-	 * - GetInputMemoryTransferStats, returns Input Memory Transfer statistics. (NumberSamples, Average, StdDev, Min, Max statistics measured in milliseconds)
+	 * - GetInferenceStats returns inference time statistics. (NumberSamples, Average, StdDev, Min, Max statistics measured in milliseconds)
+	 * - GetInputMemoryTransferStats returns Input Memory Transfer statistics. (NumberSamples, Average, StdDev, Min, Max statistics measured in
+	 *   milliseconds)
 	 */
 	float GetLastInferenceTime() const;
 	FNeuralStatsData GetInferenceStats() const;
@@ -222,8 +187,8 @@ protected:
 	/**
 	 * If DeviceType == CPU, InputDeviceType and OutputDeviceType must also be set to CPU.
 	 * If DeviceType == GPU:
-	 *  - InputDeviceType defines whether Run() will expect the input data in CPU (Run() will upload the memory to the GPU first) or GPU (no upload copy needed) format.
-	 *  - OutputDeviceType defines whether Run() will return output data in CPU (Run() will download the memory to the CPU first) or GPU (no download copy needed) format.
+	 *  - InputDeviceType: Whether Run() will expect the input data in CPU (Run will upload the memory to the GPU) or GPU (no upload needed).
+	 *  - OutputDeviceType: Whether Run() will return output data in CPU (Run will download the memory to the CPU) or GPU (no download needed).
 	 */
 	UPROPERTY(VisibleAnywhere, Category = "Neural Network Inference")
 	ENeuralDeviceType InputDeviceType;
@@ -232,21 +197,16 @@ protected:
 	ENeuralDeviceType OutputDeviceType;
 	
 	/**
-	 * SynchronousMode defines whether UNeuralNetwork::Run() will block the thread until completed (Synchronous), or whether it will run on a background thread, not
-	 * blocking the calling thread (Asynchronous).
-	 * If asynchronous, DelegateThreadMode will define whether the callback delegate is called from the game thread (highly recommended) or from any available thread (not fully thread safe).
+	 * SynchronousMode defines whether UNeuralNetwork::Run() will block the thread until completed (Synchronous), or whether it will run on a
+	 * background thread, not blocking the calling thread (Asynchronous).
+	 * If asynchronous, DelegateThreadMode will define whether the callback delegate is called from the game thread (highly recommended) or from
+	 * any available thread (not fully thread safe).
 	 * @see ENeuralNetworkSynchronousMode, ENeuralNetworkDelegateThreadMode for more details.
 	 */
 	UPROPERTY(Transient, VisibleAnywhere, Category = "Neural Network Inference")
 	ENeuralNetworkSynchronousMode SynchronousMode;
 	UPROPERTY(Transient, VisibleAnywhere, Category = "Neural Network Inference")
 	ENeuralNetworkDelegateThreadMode DelegateThreadMode;
-	
-	/**
-	 * @see ENeuralBackEnd for more details.
-	 */
-	UPROPERTY(VisibleAnywhere, Category = "Neural Network Inference")
-	ENeuralBackEnd BackEnd;
 
 	/**
 	 * Original model file path from which this UNeuralNetwork was loaded from.
@@ -258,12 +218,6 @@ private:
 	UPROPERTY(Transient, VisibleAnywhere, Category = "Neural Network Inference")
 	bool bIsLoaded;
 
-	/**
-	 * Critical section (mutex) used to avoid issues or crashes due to the asynchronous Run being run at the same time than any other non-const class function.
-	 * @see UNeuralNetwork::Run().
-	 */
-	FCriticalSection ResoucesCriticalSection;
-
 	UPROPERTY()
 	TArray<uint8> ModelReadFromFileInBytes;
 
@@ -272,18 +226,22 @@ private:
 	TArray<bool> AreInputTensorSizesVariable;
 
 	/**
-	 * If BackEnd != Auto, BackEndForCurrentPlatform will be equal to BackEnd.
-	 * Otherwise, BackEndForCurrentPlatform will be set to the optimal BackEnd given the current platform.
-	 * @see ENeuralBackEnd for more details.
+	 * Mutex to avoid issues or crashes due to the asynchronous Run() being run at the same time than any other non-const function.
+	 * @see UNeuralNetwork::Run().
 	 */
-	UPROPERTY(Transient, VisibleAnywhere, Category = "Neural Network Inference")
-	ENeuralBackEnd BackEndForCurrentPlatform;
+	FCriticalSection ResoucesCriticalSection;
 
 	/**
 	 * @see FOnAsyncRunCompleted and GetOnAsyncRunCompletedDelegate to understand OnAsyncRunCompletedDelegate.
 	 */
 	FOnAsyncRunCompleted OnAsyncRunCompletedDelegate;
 
+	/**
+	 * Stats-related members.
+	 */
+	FNeuralStats ComputeStatsModule;
+	FNeuralStats InputMemoryTransferStatsModule;
+
 	/**
 	 * Struct pointer containing the UE-and-ORT-based back end implementation.
 	 * PIMPL idiom to minimize memory when not using this back end and to to hide 3rd party dependencies.
@@ -312,15 +270,45 @@ private:
 	FNeuralTensor& GetInputTensorMutable(const int32 InTensorIndex = 0);
 	FNeuralTensor& GetOutputTensorMutable(const int32 InTensorIndex = 0);
 
-	FNeuralStats ComputeStatsModule;
-	FNeuralStats InputMemoryTransferStatsModule;
-
 public:
 	/**
 	 * Internal function not needed by the user.
 	 * Used to create custom networks without an ONNX file for QA testing in FOperatorTester::TestOperator().
 	 */
-	bool Load(TArray<FNeuralTensor>& InTensors, const TArray<FNeuralTensor*>& InInputTensors, const TArray<FNeuralTensor*>& InOutputTensors, const TArray<TSharedPtr<class FNeuralOperator>>& InOperators);
+	bool Load(TArray<FNeuralTensor>& InTensors, const TArray<FNeuralTensor*>& InInputTensors, const TArray<FNeuralTensor*>& InOutputTensors,
+		const TArray<TSharedPtr<class FNeuralOperator>>& InOperators);
+	
+	/**
+	 * Internal enum class that should not be used by the user.
+	 * Whether UNeuralNetwork will use the highly optimized UnrealEngine-and-ONNXRuntime-based back end (UEAndORT) or the less optimized but fully
+	 * cross platform only-UE one (UEOnly).
+	 * We recommend using Auto, which will find and use the optimal back end for each platform.
+	 */
+	enum class ENeuralBackEnd : uint8
+	{
+		/**
+		 * UnrealEngine-and-ONNXRuntime-accelerated back end, ideal for those platforms that support it. WITH_UE_AND_ORT_SUPPORT is the C++ define
+		 * macro that checks whether UEAndORT support exists for the current platform.
+		 */
+		UEAndORT,
+		UEOnly, /* It might be slower than the UEAndORT back end, but it will compile in all platforms and OSs compatible with Unreal Engine. */
+		Auto /* Recommended value. It will use the efficient UEAndORT if supported by the platform, and fall back to UEOnly otherwise. */
+	};
+
+	/**
+	 * Internal functions that should not be used by the user.
+	 * Getter and setter functions for BackEnd.
+	 * GetBackEnd()/GetBackEndForCurrentPlatform():
+	 * - If BackEnd == Auto, GetBackEnd() will return Auto and GetBackEndForCurrentPlatform() will return the actual BackEnd being used for the
+	 *   current platform (UEAndORT or UEOnly).
+	 * - If BackEnd != Auto, GetBackEnd() and GetBackEndForCurrentPlatform() will both return the same value (UEAndORT or UEOnly).
+	 * SetBackEnd() will modify both BackEnd and BackEndForCurrentPlatform and return IsLoaded().
+	 * SetBackEnd() is NOT THREAD SAFE, make sure that other non-const functions such as Run() are not running when SetBackEnd is called.
+	 * @see ENeuralBackEnd for more details.
+	 */
+	ENeuralBackEnd GetBackEnd() const;
+	ENeuralBackEnd GetBackEndForCurrentPlatform() const;
+	bool SetBackEnd(const ENeuralBackEnd InBackEnd);
 
 #if WITH_EDITOR
 	/**
@@ -332,6 +320,20 @@ public:
 #endif // WITH_EDITOR
 
 private:
+	/**
+	 * Internal variable that should not be used by the user.
+	 * @see ENeuralBackEnd for more details.
+	 */
+	ENeuralBackEnd BackEnd;
+
+	/**
+	 * Internal variable that should not be used by the user.
+	 * If BackEnd != Auto, BackEndForCurrentPlatform will be equal to BackEnd.
+	 * Otherwise, BackEndForCurrentPlatform will be set to the optimal BackEnd given the current platform.
+	 * @see ENeuralBackEnd for more details.
+	 */
+	ENeuralBackEnd BackEndForCurrentPlatform;
+
 #if WITH_EDITORONLY_DATA
 	/** Importing data and options used for loading the network. */
 	UPROPERTY(Instanced)
diff --git a/Engine/Plugins/Experimental/NeuralNetworkInference/Source/NeuralNetworkInference/Public/NeuralTensor.h b/Engine/Plugins/Experimental/NeuralNetworkInference/Source/NeuralNetworkInference/Public/NeuralTensor.h
--- a/Engine/Plugins/Experimental/NeuralNetworkInference/Source/NeuralNetworkInference/Public/NeuralTensor.h
+++ b/Engine/Plugins/Experimental/NeuralNetworkInference/Source/NeuralNetworkInference/Public/NeuralTensor.h
@@ -14,26 +14,6 @@ template<typename T> class TRefCountPtr;
 enum class EBufferUsageFlags : uint32;
 struct FTensorProto;
 
-
-/**
- * Although conceptually this could apply to both the CPU and GPU versions, in practice only the GPU performance is affected by this setting.
- * Input and Intermediate(Not)Initialized currently share the same attributes because input might become intermediate (e.g., if input tensor fed into a ReLU, which simply modifies
- * the input FNeuralTensor). However, Intermediate(Not)Initialized and Output do not copy the memory from CPU to GPU but rather simply allocates it.
- * Output might also become Intermediate(Not)Initialized (e.g., if Output -> ReLU -> Output), so it is kept as ReadWrite rather than written once to account for this.
- */
-UENUM()
-enum class ENeuralTensorTypeGPU : uint8
-{
-	Generic,					/** Generic tensor that works in every situation (ReadWrite), although it might not be the most efficient one. */
-	Input,						/** Input tensor of the UNeuralNetworkLegacy. Copied from CPU and ReadWrite (but usually ReadOnly). */
-	IntermediateNotInitialized,	/** Intermediate tensor of the UNeuralNetworkLegacy (output of at least a layer and input of at least some other layer). Not copied from CPU, ReadWrite, and transient. */
-	IntermediateInitialized,	/** Intermediate tensor that is initialized with CPU data (e.g., XWithZeros in FConvTranpose). Copied from CPU. */
-	Output,						/** Output tensor of the UNeuralNetworkLegacy. Not copied from CPU and ReadWrite. */
-	Weight						/** Weights of a particular operator/layer. Copied from CPU, ReadOnly, and initialized from CPU memory. */
-};
-
-
-
 /**
  * This is an auxiliary class. See UNeuralNetworkLegacy for a high-level wrapper of the whole NeuralNetworkInference plugin. The UNeuralNetworkLegacy header
  * documentation also includes some code examples.
diff --git a/Engine/Plugins/Experimental/NeuralNetworkInference/Source/NeuralNetworkInferenceQA/Private/ModelUnitTester.cpp b/Engine/Plugins/Experimental/NeuralNetworkInference/Source/NeuralNetworkInferenceQA/Private/ModelUnitTester.cpp
--- a/Engine/Plugins/Experimental/NeuralNetworkInference/Source/NeuralNetworkInferenceQA/Private/ModelUnitTester.cpp
+++ b/Engine/Plugins/Experimental/NeuralNetworkInference/Source/NeuralNetworkInferenceQA/Private/ModelUnitTester.cpp
@@ -169,13 +169,13 @@ bool FModelUnitTester::ModelLoadAccuracyAndSpeedTests(const FString& InProjectCo
 		const bool bShouldRunUEAndORTBackEnd = (InCPURepetitionsForUEAndORTBackEnd[ModelIndex] + InGPURepetitionsForUEAndORTBackEnd[ModelIndex] > 0);
 		if (bShouldRunUEAndORTBackEnd)
 		{
-			bDidGlobalTestPassed &= ModelAccuracyTest(Network, ENeuralNetworkSynchronousMode::Synchronous, ENeuralBackEnd::UEAndORT, InInputArrayValues, CPUGroundTruths, GPUGroundTruths);
-			bDidGlobalTestPassed &= ModelAccuracyTest(Network, ENeuralNetworkSynchronousMode::Asynchronous, ENeuralBackEnd::UEAndORT, InInputArrayValues, CPUGroundTruths, GPUGroundTruths);
+			bDidGlobalTestPassed &= ModelAccuracyTest(Network, ENeuralNetworkSynchronousMode::Synchronous, UNeuralNetwork::ENeuralBackEnd::UEAndORT, InInputArrayValues, CPUGroundTruths, GPUGroundTruths);
+			bDidGlobalTestPassed &= ModelAccuracyTest(Network, ENeuralNetworkSynchronousMode::Asynchronous, UNeuralNetwork::ENeuralBackEnd::UEAndORT, InInputArrayValues, CPUGroundTruths, GPUGroundTruths);
 		}
 		const bool bShouldRunUEOnlyBackEnd = false; // (InCPURepetitionsForUEOnlyBackEnd[ModelIndex] * InCPURepetitionsForUEOnlyBackEnd[ModelIndex] > 0);
 		if (bShouldRunUEOnlyBackEnd)
 		{
-			bDidGlobalTestPassed &= ModelAccuracyTest(Network, ENeuralNetworkSynchronousMode::Synchronous, ENeuralBackEnd::UEOnly, InInputArrayValues, CPUGroundTruths, GPUGroundTruths);
+			bDidGlobalTestPassed &= ModelAccuracyTest(Network, ENeuralNetworkSynchronousMode::Synchronous, UNeuralNetwork::ENeuralBackEnd::UEOnly, InInputArrayValues, CPUGroundTruths, GPUGroundTruths);
 		}
 
 		UE_LOG(LogNeuralNetworkInferenceQA, Display, TEXT("---------------------------------------------------------------------------------------------------------------------------------"));
@@ -189,12 +189,12 @@ bool FModelUnitTester::ModelLoadAccuracyAndSpeedTests(const FString& InProjectCo
 			UE_LOG(LogNeuralNetworkInferenceQA, Display, TEXT("-------------------- %s - Network %s Load and Run - %s"), *ModelName, *ModelType, *ModelFilePath);
 			if (bShouldRunUEAndORTBackEnd)
 			{
-				bDidGlobalTestPassed &= ModelAccuracyTest(NetworkONNXOrORTLoadTest(ModelFilePath), ENeuralNetworkSynchronousMode::Synchronous, ENeuralBackEnd::UEAndORT, InInputArrayValues, CPUGroundTruths, GPUGroundTruths);
-				bDidGlobalTestPassed &= ModelAccuracyTest(NetworkONNXOrORTLoadTest(ModelFilePath), ENeuralNetworkSynchronousMode::Asynchronous, ENeuralBackEnd::UEAndORT, InInputArrayValues, CPUGroundTruths, GPUGroundTruths);
+				bDidGlobalTestPassed &= ModelAccuracyTest(NetworkONNXOrORTLoadTest(ModelFilePath), ENeuralNetworkSynchronousMode::Synchronous, UNeuralNetwork::ENeuralBackEnd::UEAndORT, InInputArrayValues, CPUGroundTruths, GPUGroundTruths);
+				bDidGlobalTestPassed &= ModelAccuracyTest(NetworkONNXOrORTLoadTest(ModelFilePath), ENeuralNetworkSynchronousMode::Asynchronous, UNeuralNetwork::ENeuralBackEnd::UEAndORT, InInputArrayValues, CPUGroundTruths, GPUGroundTruths);
 			}
 			if (bShouldRunUEOnlyBackEnd)
 			{
-				bDidGlobalTestPassed &= ModelAccuracyTest(NetworkONNXOrORTLoadTest(ModelFilePath), ENeuralNetworkSynchronousMode::Synchronous, ENeuralBackEnd::UEOnly, InInputArrayValues, CPUGroundTruths, GPUGroundTruths);
+				bDidGlobalTestPassed &= ModelAccuracyTest(NetworkONNXOrORTLoadTest(ModelFilePath), ENeuralNetworkSynchronousMode::Synchronous, UNeuralNetwork::ENeuralBackEnd::UEOnly, InInputArrayValues, CPUGroundTruths, GPUGroundTruths);
 			}
 		}
 #else //WITH_EDITOR
@@ -211,16 +211,16 @@ bool FModelUnitTester::ModelLoadAccuracyAndSpeedTests(const FString& InProjectCo
 
 		// UEAndORT (if WITH_UE_AND_ORT_SUPPORT)
 #ifdef WITH_UE_AND_ORT_SUPPORT
-		bDidGlobalTestPassed &= ModelSpeedTest(UAssetModelFilePath, ENeuralDeviceType::CPU, ENeuralBackEnd::UEAndORT, InCPURepetitionsForUEAndORTBackEnd[ModelIndex]);
+		bDidGlobalTestPassed &= ModelSpeedTest(UAssetModelFilePath, ENeuralDeviceType::CPU, UNeuralNetwork::ENeuralBackEnd::UEAndORT, InCPURepetitionsForUEAndORTBackEnd[ModelIndex]);
 		UNeuralNetwork* Network = NetworkUassetLoadTest(UAssetModelFilePath);
-		if (Network->GetBackEndForCurrentPlatform() != ENeuralBackEnd::UEAndORT)
+		if (Network->GetBackEndForCurrentPlatform() != UNeuralNetwork::ENeuralBackEnd::UEAndORT)
 		{
 			UE_LOG(LogNeuralNetworkInferenceQA, Warning, TEXT("-------------------- Default UAsset BackEnd should be UEAndORT."));
 			return false;
 		}
 		else if (InGPURepetitionsForUEAndORTBackEnd[ModelIndex] > 0 && Network->IsGPUSupported())
 		{
-			bDidGlobalTestPassed &= ModelSpeedTest(UAssetModelFilePath, ENeuralDeviceType::GPU, ENeuralBackEnd::UEAndORT, InGPURepetitionsForUEAndORTBackEnd[ModelIndex]);
+			bDidGlobalTestPassed &= ModelSpeedTest(UAssetModelFilePath, ENeuralDeviceType::GPU, UNeuralNetwork::ENeuralBackEnd::UEAndORT, InGPURepetitionsForUEAndORTBackEnd[ModelIndex]);
 		}
 		else
 		{
@@ -233,8 +233,8 @@ bool FModelUnitTester::ModelLoadAccuracyAndSpeedTests(const FString& InProjectCo
 		const bool bShouldRunUEOnlyBackEnd = false;
 		if (bShouldRunUEOnlyBackEnd)
 		{
-			bDidGlobalTestPassed &= ModelSpeedTest(UAssetModelFilePath, ENeuralDeviceType::CPU, ENeuralBackEnd::UEOnly, InCPURepetitionsForUEOnlyBackEnd[ModelIndex]);
-			bDidGlobalTestPassed &= ModelSpeedTest(UAssetModelFilePath, ENeuralDeviceType::GPU, ENeuralBackEnd::UEOnly, InGPURepetitionsForUEOnlyBackEnd[ModelIndex]);
+			bDidGlobalTestPassed &= ModelSpeedTest(UAssetModelFilePath, ENeuralDeviceType::CPU, UNeuralNetwork::ENeuralBackEnd::UEOnly, InCPURepetitionsForUEOnlyBackEnd[ModelIndex]);
+			bDidGlobalTestPassed &= ModelSpeedTest(UAssetModelFilePath, ENeuralDeviceType::GPU, UNeuralNetwork::ENeuralBackEnd::UEOnly, InGPURepetitionsForUEOnlyBackEnd[ModelIndex]);
 		}
 	}
 
@@ -289,7 +289,7 @@ UNeuralNetwork* FModelUnitTester::NetworkONNXOrORTLoadTest(const FString& InMode
 	return Network;
 }
 
-bool FModelUnitTester::ModelAccuracyTest(UNeuralNetwork* InOutNetwork, const ENeuralNetworkSynchronousMode InSynchronousMode, const ENeuralBackEnd InBackEnd, const TArray<float>& InInputArrayValues,
+bool FModelUnitTester::ModelAccuracyTest(UNeuralNetwork* InOutNetwork, const ENeuralNetworkSynchronousMode InSynchronousMode, const UNeuralNetwork::ENeuralBackEnd InBackEnd, const TArray<float>& InInputArrayValues,
 	const TArray<double>& InCPUGroundTruths, const TArray<double>& InGPUGroundTruths)
 {
 	// Sanity check
@@ -322,7 +322,7 @@ bool FModelUnitTester::ModelAccuracyTest(UNeuralNetwork* InOutNetwork, const ENe
 	const ENeuralDeviceType OriginalOutputDeviceType = InOutNetwork->GetOutputDeviceType();
 	const ENeuralNetworkSynchronousMode OriginalSynchronousMode = InOutNetwork->GetSynchronousMode();
 	const ENeuralNetworkDelegateThreadMode OriginalDelegateThreadMode = InOutNetwork->GetOnAsyncRunCompletedDelegateMode();
-	const ENeuralBackEnd OriginalBackEnd = InOutNetwork->GetBackEnd();
+	const UNeuralNetwork::ENeuralBackEnd OriginalBackEnd = InOutNetwork->GetBackEnd();
 	
 	// Set (a)synchronous Mode
 	InOutNetwork->SetSynchronousMode(InSynchronousMode);
@@ -462,7 +462,7 @@ void FModelUnitTester::ModelAccuracyTestRun(TArray<TArray<float>>& OutOutputs, U
 	}
 }
 
-bool FModelUnitTester::ModelSpeedTest(const FString& InUAssetPath, const ENeuralDeviceType InDeviceType, const ENeuralBackEnd InBackEnd, const int32 InRepetitions)
+bool FModelUnitTester::ModelSpeedTest(const FString& InUAssetPath, const ENeuralDeviceType InDeviceType, const UNeuralNetwork::ENeuralBackEnd InBackEnd, const int32 InRepetitions)
 {
 	// Get debug strings
 	const FString DeviceTypeString = GetDeviceTypeString(InDeviceType);
@@ -483,7 +483,7 @@ bool FModelUnitTester::ModelSpeedTest(const FString& InUAssetPath, const ENeural
 	}
 	// Save original network state
 	const ENeuralDeviceType OriginalDeviceType = InOutNetwork->GetDeviceType();
-	const ENeuralBackEnd OriginalBackEnd = InOutNetwork->GetBackEnd();
+	const UNeuralNetwork::ENeuralBackEnd OriginalBackEnd = InOutNetwork->GetBackEnd();
 	// Set desired back end
 	if (!InOutNetwork->SetBackEnd(InBackEnd))
 	{
@@ -564,13 +564,13 @@ FString FModelUnitTester::GetDeviceTypeString(const ENeuralDeviceType InDeviceTy
 	}
 }
 
-FString FModelUnitTester::GetBackEndString(const ENeuralBackEnd InBackEnd)
+FString FModelUnitTester::GetBackEndString(const UNeuralNetwork::ENeuralBackEnd InBackEnd)
 {
-	if (InBackEnd == ENeuralBackEnd::UEAndORT)
+	if (InBackEnd == UNeuralNetwork::ENeuralBackEnd::UEAndORT)
 	{
 		return TEXT("UEAndORT");
 	}
-	else if (InBackEnd == ENeuralBackEnd::UEOnly)
+	else if (InBackEnd == UNeuralNetwork::ENeuralBackEnd::UEOnly)
 	{
 		return TEXT("UEOnly");
 	}
diff --git a/Engine/Plugins/Experimental/NeuralNetworkInference/Source/NeuralNetworkInferenceQA/Private/ModelUnitTester.h b/Engine/Plugins/Experimental/NeuralNetworkInference/Source/NeuralNetworkInferenceQA/Private/ModelUnitTester.h
--- a/Engine/Plugins/Experimental/NeuralNetworkInference/Source/NeuralNetworkInferenceQA/Private/ModelUnitTester.h
+++ b/Engine/Plugins/Experimental/NeuralNetworkInference/Source/NeuralNetworkInferenceQA/Private/ModelUnitTester.h
@@ -26,16 +26,16 @@ private:
 	static FString GetUAssetModelFilePath(const FString& InModelName, const FString& InModelZooRelativeDirectory);
 	static UNeuralNetwork* NetworkUassetLoadTest(const FString& InUAssetPath);
 	static UNeuralNetwork* NetworkONNXOrORTLoadTest(const FString& InModelFilePath);
-	static bool ModelAccuracyTest(UNeuralNetwork* InOutNetwork, const ENeuralNetworkSynchronousMode InSynchronousMode, const ENeuralBackEnd InBackEnd, const TArray<float>& InInputArrayValues,
+	static bool ModelAccuracyTest(UNeuralNetwork* InOutNetwork, const ENeuralNetworkSynchronousMode InSynchronousMode, const UNeuralNetwork::ENeuralBackEnd InBackEnd, const TArray<float>& InInputArrayValues,
 		const TArray<double>& InCPUGroundTruths, const TArray<double>& InGPUGroundTruths);
 	static void ModelAccuracyTestRun(TArray<TArray<float>>& OutOutputs, UNeuralNetwork* InOutNetwork, const TArray<TArray<float>> InInputArrays,
 		const ENeuralDeviceType InDeviceType, const ENeuralDeviceType InInputDeviceType, const ENeuralDeviceType InOutputDeviceType);
-	static bool ModelSpeedTest(const FString& InUAssetPath, const ENeuralDeviceType InDeviceType, const ENeuralBackEnd InBackEnd, const int32 InRepetitions);
+	static bool ModelSpeedTest(const FString& InUAssetPath, const ENeuralDeviceType InDeviceType, const UNeuralNetwork::ENeuralBackEnd InBackEnd, const int32 InRepetitions);
 	static double GetAveragedL1Norm(const TArray<float>& InArray);
 	static double GetAveragedL1NormDiff(const TArray<float>& InArray1, const TArray<float>& InArray2);
 	/**
 	 * Verbose string auxiliary functions.
 	 */
 	static FString GetDeviceTypeString(const ENeuralDeviceType InDeviceType);
-	static FString GetBackEndString(const ENeuralBackEnd InBackEnd);
+	static FString GetBackEndString(const UNeuralNetwork::ENeuralBackEnd InBackEnd);
 };
