Adding NNI documentation - Part 2:
- NNI/Public/NeuralEnumClasses.h fully commented and some functions renamed/moved.

#rnx

diff --git a/Engine/Plugins/Experimental/NeuralNetworkInference/Source/NeuralNetworkInference/Private/NeuralEnumClasses.cpp b/Engine/Plugins/Experimental/NeuralNetworkInference/Source/NeuralNetworkInference/Private/NeuralEnumClasses.cpp
--- a/Engine/Plugins/Experimental/NeuralNetworkInference/Source/NeuralNetworkInference/Private/NeuralEnumClasses.cpp
+++ b/Engine/Plugins/Experimental/NeuralNetworkInference/Source/NeuralNetworkInference/Private/NeuralEnumClasses.cpp
@@ -8,37 +8,7 @@
 /* FNeuralNetworkInferenceUtils public functions
  *****************************************************************************/
 
-FString FNeuralDataTypeUtils::ToString(const ENeuralDataType InDataType)
-{
-	if (InDataType == ENeuralDataType::Float)
-	{
-		return TEXT("Float");
-	}
-	else if (InDataType == ENeuralDataType::Int32)
-	{
-		return TEXT("Int32");
-	}
-	else if (InDataType == ENeuralDataType::Int64)
-	{
-		return TEXT("Int64");
-	}
-	else if (InDataType == ENeuralDataType::UInt32)
-	{
-		return TEXT("UInt32");
-	}
-	else if (InDataType == ENeuralDataType::UInt64)
-	{
-		return TEXT("UInt64");
-	}
-	else if (InDataType == ENeuralDataType::None)
-	{
-		return TEXT("None");
-	}
-	UE_LOG(LogNeuralNetworkInference, Warning, TEXT("FNeuralDataTypeUtils::ToString(): Unknown InDataType = %d used."), (int32)InDataType);
-	return TEXT("");
-}
-
-int64 FNeuralDataTypeUtils::GetSize(const ENeuralDataType InDataType)
+int64 FNeuralDataTypeUtils::GetByteSize(const ENeuralDataType InDataType)
 {
 	if (InDataType == ENeuralDataType::Float)
 	{
@@ -60,7 +30,7 @@ int64 FNeuralDataTypeUtils::GetSize(const ENeuralDataType InDataType)
 	{
 		return sizeof(uint64);
 	}
-	UE_LOG(LogNeuralNetworkInference, Warning, TEXT("FNeuralDataTypeUtils::GetSize(): Unknown InDataType = %d used."), (int32)InDataType);
+	UE_LOG(LogNeuralNetworkInference, Warning, TEXT("FNeuralDataTypeUtils::GetByteSize(): Unknown InDataType = %d used."), (int32)InDataType);
 	return 1;
 }
 
@@ -81,3 +51,33 @@ EPixelFormat FNeuralDataTypeUtils::GetPixelFormat(const ENeuralDataType InDataTy
 	UE_LOG(LogNeuralNetworkInference, Warning, TEXT("FNeuralDataTypeUtils::GetPixelFormat(): Unknown InDataType = %d used."), (int32)InDataType);
 	return EPixelFormat::PF_Unknown;
 }
+
+FString FNeuralDataTypeUtils::ToString(const ENeuralDataType InDataType)
+{
+	if (InDataType == ENeuralDataType::Float)
+	{
+		return TEXT("Float");
+	}
+	else if (InDataType == ENeuralDataType::Int32)
+	{
+		return TEXT("Int32");
+	}
+	else if (InDataType == ENeuralDataType::Int64)
+	{
+		return TEXT("Int64");
+	}
+	else if (InDataType == ENeuralDataType::UInt32)
+	{
+		return TEXT("UInt32");
+	}
+	else if (InDataType == ENeuralDataType::UInt64)
+	{
+		return TEXT("UInt64");
+	}
+	else if (InDataType == ENeuralDataType::None)
+	{
+		return TEXT("None");
+	}
+	UE_LOG(LogNeuralNetworkInference, Warning, TEXT("FNeuralDataTypeUtils::ToString(): Unknown InDataType = %d used."), (int32)InDataType);
+	return TEXT("");
+}
diff --git a/Engine/Plugins/Experimental/NeuralNetworkInference/Source/NeuralNetworkInference/Private/NeuralNetwork.cpp b/Engine/Plugins/Experimental/NeuralNetworkInference/Source/NeuralNetworkInference/Private/NeuralNetwork.cpp
--- a/Engine/Plugins/Experimental/NeuralNetworkInference/Source/NeuralNetworkInference/Private/NeuralNetwork.cpp
+++ b/Engine/Plugins/Experimental/NeuralNetworkInference/Source/NeuralNetworkInference/Private/NeuralNetwork.cpp
@@ -48,8 +48,8 @@ UNeuralNetwork::UNeuralNetwork()
 	: DeviceType(ENeuralDeviceType::GPU)
 	, InputDeviceType(ENeuralDeviceType::CPU)
 	, OutputDeviceType(ENeuralDeviceType::CPU)
-	, SynchronousMode(ENeuralNetworkSynchronousMode::Synchronous)
-	, DelegateThreadMode(ENeuralNetworkDelegateThreadMode::GameThread)
+	, SynchronousMode(ENeuralSynchronousMode::Synchronous)
+	, DelegateThreadMode(ENeuralThreadMode::GameThread)
 	, bIsLoaded(false)
 	, BackEnd(ENeuralBackEnd::Auto)
 	, BackEndForCurrentPlatform(FPrivateNeuralNetwork::SetBackEndForCurrentPlatform(BackEnd))
@@ -159,12 +159,12 @@ void UNeuralNetwork::SetDeviceType(const ENeuralDeviceType InDeviceType, const E
 	}
 }
 
-ENeuralNetworkSynchronousMode UNeuralNetwork::GetSynchronousMode() const
+ENeuralSynchronousMode UNeuralNetwork::GetSynchronousMode() const
 {
 	return SynchronousMode;
 }
 
-void UNeuralNetwork::SetSynchronousMode(const ENeuralNetworkSynchronousMode InSynchronousMode)
+void UNeuralNetwork::SetSynchronousMode(const ENeuralSynchronousMode InSynchronousMode)
 {
 	const FScopeLock ResourcesLock(&ResoucesCriticalSection);
 	SynchronousMode = InSynchronousMode;
@@ -176,12 +176,12 @@ UNeuralNetwork::FOnAsyncRunCompleted& UNeuralNetwork::GetOnAsyncRunCompletedDele
 	return OnAsyncRunCompletedDelegate;
 }
 
-ENeuralNetworkDelegateThreadMode UNeuralNetwork::GetOnAsyncRunCompletedDelegateMode() const
+ENeuralThreadMode UNeuralNetwork::GetOnAsyncRunCompletedDelegateMode() const
 {
 	return DelegateThreadMode;
 }
 
-void UNeuralNetwork::SetOnAsyncRunCompletedDelegateMode(const ENeuralNetworkDelegateThreadMode InDelegateThreadMode)
+void UNeuralNetwork::SetOnAsyncRunCompletedDelegateMode(const ENeuralThreadMode InDelegateThreadMode)
 {
 	const FScopeLock ResourcesLock(&ResoucesCriticalSection);
 	DelegateThreadMode = InDelegateThreadMode;
diff --git a/Engine/Plugins/Experimental/NeuralNetworkInference/Source/NeuralNetworkInference/Private/NeuralNetworkImplBackEndUEAndORT.cpp b/Engine/Plugins/Experimental/NeuralNetworkInference/Source/NeuralNetworkInference/Private/NeuralNetworkImplBackEndUEAndORT.cpp
--- a/Engine/Plugins/Experimental/NeuralNetworkInference/Source/NeuralNetworkInference/Private/NeuralNetworkImplBackEndUEAndORT.cpp
+++ b/Engine/Plugins/Experimental/NeuralNetworkInference/Source/NeuralNetworkInference/Private/NeuralNetworkImplBackEndUEAndORT.cpp
@@ -235,7 +235,7 @@ UNeuralNetwork::FImplBackEndUEAndORT::FNeuralNetworkAsyncTask::FNeuralNetworkAsy
 	: BackEnd(InBackEnd)
 {}
 
-void UNeuralNetwork::FImplBackEndUEAndORT::FNeuralNetworkAsyncTask::SetRunSessionArgs(const ENeuralNetworkSynchronousMode InSyncMode, const ENeuralDeviceType InDeviceType,
+void UNeuralNetwork::FImplBackEndUEAndORT::FNeuralNetworkAsyncTask::SetRunSessionArgs(const ENeuralSynchronousMode InSyncMode, const ENeuralDeviceType InDeviceType,
 	const ENeuralDeviceType InInputDeviceType, const ENeuralDeviceType InOutputDeviceType)
 {
 	const FScopeLock ResourcesLock(&BackEnd->ResoucesCriticalSection);
@@ -248,7 +248,7 @@ void UNeuralNetwork::FImplBackEndUEAndORT::FNeuralNetworkAsyncTask::SetRunSessio
 void UNeuralNetwork::FImplBackEndUEAndORT::FNeuralNetworkAsyncTask::DoWork()
 {
 	const FRedirectCoutAndCerrToUeLog RedirectCoutAndCerrToUeLog;
-	if (SyncMode == ENeuralNetworkSynchronousMode::Asynchronous)
+	if (SyncMode == ENeuralSynchronousMode::Asynchronous)
 	{
 		BackEnd->RunSessionAsync(DeviceType, InputDeviceType, OutputDeviceType);
 	}
@@ -264,7 +264,7 @@ void UNeuralNetwork::FImplBackEndUEAndORT::FNeuralNetworkAsyncTask::DoWork()
 /* FImplBackEndUEAndORT 'structors
  *****************************************************************************/
 
-UNeuralNetwork::FImplBackEndUEAndORT::FImplBackEndUEAndORT(FOnAsyncRunCompleted& InOutOnAsyncRunCompletedDelegate, ENeuralNetworkDelegateThreadMode& InDelegateThreadMode,
+UNeuralNetwork::FImplBackEndUEAndORT::FImplBackEndUEAndORT(FOnAsyncRunCompleted& InOutOnAsyncRunCompletedDelegate, ENeuralThreadMode& InDelegateThreadMode,
 	FCriticalSection& InResoucesCriticalSection)
 #ifdef WITH_UE_AND_ORT_SUPPORT
 	: OnAsyncRunCompletedDelegate(InOutOnAsyncRunCompletedDelegate)
@@ -334,7 +334,7 @@ bool UNeuralNetwork::FImplBackEndUEAndORT::IsGPUSupported()
 }
 
 bool UNeuralNetwork::FImplBackEndUEAndORT::Load(TSharedPtr<FImplBackEndUEAndORT>& InOutImplBackEndUEAndORT,
-	FOnAsyncRunCompleted& InOutOnAsyncRunCompletedDelegate, ENeuralNetworkDelegateThreadMode& InOutDelegateThreadMode, FCriticalSection& InOutResoucesCriticalSection,
+	FOnAsyncRunCompleted& InOutOnAsyncRunCompletedDelegate, ENeuralThreadMode& InOutDelegateThreadMode, FCriticalSection& InOutResoucesCriticalSection,
 	TArray<bool>& OutAreInputTensorSizesVariable, const TArray<uint8>& InModelReadFromFileInBytes, const FString& InModelFullFilePath,
 	const ENeuralDeviceType InDeviceType, const ENeuralDeviceType InInputDeviceType, const ENeuralDeviceType InOutputDeviceType)
 {
@@ -421,7 +421,7 @@ bool UNeuralNetwork::FImplBackEndUEAndORT::Load(TSharedPtr<FImplBackEndUEAndORT>
 #endif //WITH_UE_AND_ORT_SUPPORT
 }
 
-void UNeuralNetwork::FImplBackEndUEAndORT::Run(const ENeuralNetworkSynchronousMode InSynchronousMode, const ENeuralDeviceType InDeviceType,
+void UNeuralNetwork::FImplBackEndUEAndORT::Run(const ENeuralSynchronousMode InSynchronousMode, const ENeuralDeviceType InDeviceType,
 	const ENeuralDeviceType InInputDeviceType, const ENeuralDeviceType InOutputDeviceType)
 {
 #ifdef WITH_UE_AND_ORT_SUPPORT
@@ -433,11 +433,11 @@ void UNeuralNetwork::FImplBackEndUEAndORT::Run(const ENeuralNetworkSynchronousMo
 		NeuralNetworkAsyncTask->GetTask().SetRunSessionArgs(InSynchronousMode, InDeviceType, InInputDeviceType, InOutputDeviceType);
 
 		// Run UNeuralNetwork
-		if (InSynchronousMode == ENeuralNetworkSynchronousMode::Synchronous)
+		if (InSynchronousMode == ENeuralSynchronousMode::Synchronous)
 		{
 			NeuralNetworkAsyncTask->StartSynchronousTask();
 		}
-		else if (InSynchronousMode == ENeuralNetworkSynchronousMode::Asynchronous)
+		else if (InSynchronousMode == ENeuralSynchronousMode::Asynchronous)
 		{
 			NeuralNetworkAsyncTask->StartBackgroundTask(); // Alternative: (GThreadPool, EQueuedWorkPriority::Highest);
 		}
@@ -474,7 +474,7 @@ void UNeuralNetwork::FImplBackEndUEAndORT::EnsureAsyncTaskCompletion() const
 }
 
 bool UNeuralNetwork::FImplBackEndUEAndORT::InitializedAndConfigureMembers(TSharedPtr<FImplBackEndUEAndORT>& InOutImplBackEndUEAndORT,
-	FOnAsyncRunCompleted& InOutOnAsyncRunCompletedDelegate, ENeuralNetworkDelegateThreadMode& InOutDelegateThreadMode, FCriticalSection& InOutResoucesCriticalSection,
+	FOnAsyncRunCompleted& InOutOnAsyncRunCompletedDelegate, ENeuralThreadMode& InOutDelegateThreadMode, FCriticalSection& InOutResoucesCriticalSection,
 	const FString& InModelFullFilePath, const ENeuralDeviceType InDeviceType)
 {
 	// Initialize InOutImplBackEndUEAndORT
@@ -628,7 +628,7 @@ bool UNeuralNetwork::FImplBackEndUEAndORT::ConfigureTensors(TArray<FNeuralTensor
 	TArray<const char*> TensorNames;
 	TArray<ENeuralDataType> TensorDataTypes;
 	TArray<TArray<int64>> TensorSizes;
-	TArray<ENeuralTensorTypeGPU> TensorGPUTypes;
+	TArray<ENeuralTensorType> TensorGPUTypes;
 
 	const uint32 NumberTensors = bIsInput ? Session->GetInputCount() : Session->GetOutputCount();
 	if (OutAreInputTensorSizesVariable)
@@ -710,22 +710,22 @@ bool UNeuralNetwork::FImplBackEndUEAndORT::ConfigureTensors(TArray<FNeuralTensor
 		// @todo: Should caller specify tensor GPU type?
 		// Input/Output tensor GPU type means that data should not be copied from CPU
 		// Generic means that data is on the CPU and it will be copied to GPU
-		ENeuralTensorTypeGPU	TensorGPUType;
+		ENeuralTensorType TensorGPUType;
 
 		if (InDeviceType == ENeuralDeviceType::GPU)
 		{
 			if (bIsInput)
 			{
-				TensorGPUType = InInputDeviceType == ENeuralDeviceType::GPU ? ENeuralTensorTypeGPU::Input : ENeuralTensorTypeGPU::Generic;
+				TensorGPUType = InInputDeviceType == ENeuralDeviceType::GPU ? ENeuralTensorType::Input : ENeuralTensorType::Generic;
 			}
 			else
 			{
-				TensorGPUType = InOutputDeviceType == ENeuralDeviceType::GPU ? ENeuralTensorTypeGPU::Output : ENeuralTensorTypeGPU::Generic;
+				TensorGPUType = InOutputDeviceType == ENeuralDeviceType::GPU ? ENeuralTensorType::Output : ENeuralTensorType::Generic;
 			}
 		}
 		else
 		{
-			TensorGPUType = ENeuralTensorTypeGPU::Generic;
+			TensorGPUType = ENeuralTensorType::Generic;
 		}
 
 		TensorGPUTypes.Push(TensorGPUType);
@@ -737,7 +737,7 @@ bool UNeuralNetwork::FImplBackEndUEAndORT::ConfigureTensors(TArray<FNeuralTensor
 }
 
 bool UNeuralNetwork::FImplBackEndUEAndORT::SetTensorsFromNetwork(TArray<FNeuralTensor>& OutTensors, TArray<const char*>& InTensorNames,
-	TArray<ENeuralDataType>& InTensorDataTypes, TArray<TArray<int64>>& InSizes, TArray<ENeuralTensorTypeGPU>& InTensorGPUTypes, const bool bIsInput)
+	TArray<ENeuralDataType>& InTensorDataTypes, TArray<TArray<int64>>& InSizes, TArray<ENeuralTensorType>& InTensorGPUTypes, const bool bIsInput)
 {
 	const int32 TensorNumber = InTensorNames.Num();
 	if (InTensorDataTypes.Num() != TensorNumber || InSizes.Num() != TensorNumber)
@@ -801,14 +801,14 @@ bool UNeuralNetwork::FImplBackEndUEAndORT::SetTensorsFromNetwork(TArray<FNeuralT
 		}
 
 #ifdef PLATFORM_WIN64
-		if (InTensorGPUTypes[TensorIndex] == ENeuralTensorTypeGPU::Generic)
+		if (InTensorGPUTypes[TensorIndex] == ENeuralTensorType::Generic)
 		{
 			// Pre-allocate TArray (if size is different)
 			OutTensors[TensorIndex].SetNumUninitialized(InSizes[TensorIndex], InTensorDataTypes[TensorIndex]);
 			// Link tensor with ORT blob
 			LinkTensorToONNXRuntime(OutTensors, OrtTensors, *AllocatorInfo, TensorIndex);
 		}
-		else if (InTensorGPUTypes[TensorIndex] == ENeuralTensorTypeGPU::Input || InTensorGPUTypes[TensorIndex] == ENeuralTensorTypeGPU::Output)
+		else if (InTensorGPUTypes[TensorIndex] == ENeuralTensorType::Input || InTensorGPUTypes[TensorIndex] == ENeuralTensorType::Output)
 		{
 			// @todo: should we remove this? It's currently used to read memory from GPU to CPU
 			OutTensors[TensorIndex].SetNumUninitialized(InSizes[TensorIndex], InTensorDataTypes[TensorIndex]);
@@ -943,7 +943,7 @@ void UNeuralNetwork::FImplBackEndUEAndORT::RunSessionAsync(const ENeuralDeviceTy
 	RunSessionImpl(InDeviceType, InInputDeviceType, InOutputDeviceType);
 
 	// Execute OnAsyncRunCompletedDelegate on main thread
-	if (DelegateThreadMode == ENeuralNetworkDelegateThreadMode::GameThread)
+	if (DelegateThreadMode == ENeuralThreadMode::GameThread)
 	{
 		if (OnAsyncRunCompletedDelegate.IsBound())
 		{
@@ -959,7 +959,7 @@ void UNeuralNetwork::FImplBackEndUEAndORT::RunSessionAsync(const ENeuralDeviceTy
 			}
 		}
 	}
-	else if (DelegateThreadMode == ENeuralNetworkDelegateThreadMode::AnyThread)
+	else if (DelegateThreadMode == ENeuralThreadMode::AnyThread)
 	{
 		OnAsyncRunCompletedDelegate.ExecuteIfBound();
 	}
@@ -992,7 +992,7 @@ void UNeuralNetwork::FImplBackEndUEAndORT::RunSessionImpl(const ENeuralDeviceTyp
 
 		for (auto& InputTensor : InputTensors)
 		{
-			if (InputTensor.GetTensorTypeGPU() != ENeuralTensorTypeGPU::Input)
+			if (InputTensor.GetTensorTypeGPU() != ENeuralTensorType::Input)
 				continue;
 
 			bNeedsGPUCopy = true;
diff --git a/Engine/Plugins/Experimental/NeuralNetworkInference/Source/NeuralNetworkInference/Private/NeuralNetworkImplBackEndUEAndORT.h b/Engine/Plugins/Experimental/NeuralNetworkInference/Source/NeuralNetworkInference/Private/NeuralNetworkImplBackEndUEAndORT.h
--- a/Engine/Plugins/Experimental/NeuralNetworkInference/Source/NeuralNetworkInference/Private/NeuralNetworkImplBackEndUEAndORT.h
+++ b/Engine/Plugins/Experimental/NeuralNetworkInference/Source/NeuralNetworkInference/Private/NeuralNetworkImplBackEndUEAndORT.h
@@ -43,7 +43,7 @@ public:
 	TArray<FNeuralTensor> InputTensors;
 	TArray<FNeuralTensor> OutputTensors;
 
-	FImplBackEndUEAndORT(FOnAsyncRunCompleted& InOutOnAsyncRunCompletedDelegate, ENeuralNetworkDelegateThreadMode& InDelegateThreadMode, FCriticalSection& InResoucesCriticalSection);
+	FImplBackEndUEAndORT(FOnAsyncRunCompleted& InOutOnAsyncRunCompletedDelegate, ENeuralThreadMode& InDelegateThreadMode, FCriticalSection& InResoucesCriticalSection);
 
 	~FImplBackEndUEAndORT();
 
@@ -52,18 +52,18 @@ public:
 	static bool IsGPUSupported();
 
 	static bool Load(TSharedPtr<FImplBackEndUEAndORT>& InOutImplBackEndUEAndORT, FOnAsyncRunCompleted& InOutOnAsyncRunCompletedDelegate,
-		ENeuralNetworkDelegateThreadMode& InOutDelegateThreadMode, FCriticalSection& InOutResoucesCriticalSection, TArray<bool>& OutAreInputTensorSizesVariable,
+		ENeuralThreadMode& InOutDelegateThreadMode, FCriticalSection& InOutResoucesCriticalSection, TArray<bool>& OutAreInputTensorSizesVariable,
 		const TArray<uint8>& InModelReadFromFileInBytes, const FString& InModelFullFilePath, const ENeuralDeviceType InDeviceType, const ENeuralDeviceType InInputDeviceType,
 		const ENeuralDeviceType InOutputDeviceType);
 
-	void Run(const ENeuralNetworkSynchronousMode InSynchronousMode, const ENeuralDeviceType InDeviceType, const ENeuralDeviceType InInputDeviceType,
+	void Run(const ENeuralSynchronousMode InSynchronousMode, const ENeuralDeviceType InDeviceType, const ENeuralDeviceType InInputDeviceType,
 		const ENeuralDeviceType InOutputDeviceType);
 
 #ifdef WITH_UE_AND_ORT_SUPPORT
 private:
 	/** Async support */
 	FOnAsyncRunCompleted& OnAsyncRunCompletedDelegate;
-	ENeuralNetworkDelegateThreadMode& DelegateThreadMode;
+	ENeuralThreadMode& DelegateThreadMode;
 	FCriticalSection& ResoucesCriticalSection;
 	/** Network-related variables */
 	TUniquePtr<Ort::Env> Environment;
@@ -88,7 +88,7 @@ private:
 		friend class FAsyncTask<FNeuralNetworkAsyncTask>;
 	public:
 		FNeuralNetworkAsyncTask(UNeuralNetwork::FImplBackEndUEAndORT* InBackEnd);
-		void SetRunSessionArgs(const ENeuralNetworkSynchronousMode InSyncMode, const ENeuralDeviceType InDeviceType,
+		void SetRunSessionArgs(const ENeuralSynchronousMode InSyncMode, const ENeuralDeviceType InDeviceType,
 			const ENeuralDeviceType InInputDeviceType, const ENeuralDeviceType InOutputDeviceType);
 	protected:
 		void DoWork();
@@ -100,7 +100,7 @@ private:
 	private:
 		UNeuralNetwork::FImplBackEndUEAndORT* BackEnd;
 		// Variables that could change on each inference run
-		ENeuralNetworkSynchronousMode SyncMode;
+		ENeuralSynchronousMode SyncMode;
 		ENeuralDeviceType DeviceType;
 		ENeuralDeviceType InputDeviceType;
 		ENeuralDeviceType OutputDeviceType;
@@ -110,7 +110,7 @@ private:
 	void EnsureAsyncTaskCompletion() const;
 
 	static bool InitializedAndConfigureMembers(TSharedPtr<FImplBackEndUEAndORT>& InOutImplBackEndUEAndORT, FOnAsyncRunCompleted& InOutOnAsyncRunCompletedDelegate,
-		ENeuralNetworkDelegateThreadMode& InOutDelegateThreadMode, FCriticalSection& InOutResoucesCriticalSection, const FString& InModelFullFilePath, const ENeuralDeviceType InDeviceType);
+		ENeuralThreadMode& InOutDelegateThreadMode, FCriticalSection& InOutResoucesCriticalSection, const FString& InModelFullFilePath, const ENeuralDeviceType InDeviceType);
 
 	bool ConfigureMembers(const ENeuralDeviceType InDeviceType);
 
@@ -118,7 +118,7 @@ private:
 		const ENeuralDeviceType InInputDeviceType, const ENeuralDeviceType InOutputDeviceType);
 
 	bool SetTensorsFromNetwork(TArray<FNeuralTensor>& OutTensors, TArray<const char*>& InTensorNames, TArray<ENeuralDataType>& InTensorDataTypes,
-		TArray<TArray<int64>>& InSizes, TArray<ENeuralTensorTypeGPU>& InTensorGPUTypes, const bool bIsInput);
+		TArray<TArray<int64>>& InSizes, TArray<ENeuralTensorType>& InTensorGPUTypes, const bool bIsInput);
 
 	static void LinkTensorToONNXRuntime(TArray<FNeuralTensor>& InOutTensors, TArray<Ort::Value>& InOutOrtTensors, Ort::MemoryInfo& InOutAllocatorInfo,
 		const int32 InTensorIndex);
diff --git a/Engine/Plugins/Experimental/NeuralNetworkInference/Source/NeuralNetworkInference/Private/NeuralNetworkImplBackEndUEOnly.cpp b/Engine/Plugins/Experimental/NeuralNetworkInference/Source/NeuralNetworkInference/Private/NeuralNetworkImplBackEndUEOnly.cpp
--- a/Engine/Plugins/Experimental/NeuralNetworkInference/Source/NeuralNetworkInference/Private/NeuralNetworkImplBackEndUEOnly.cpp
+++ b/Engine/Plugins/Experimental/NeuralNetworkInference/Source/NeuralNetworkInference/Private/NeuralNetworkImplBackEndUEOnly.cpp
@@ -87,7 +87,7 @@ bool UNeuralNetwork::FImplBackEndUEOnly::Load(TSharedPtr<FImplBackEndUEOnly>& In
 //	return (InOutImplBackEndUEOnly->Operators.Num() > 0 && InOutImplBackEndUEOnly->TensorManager.IsLoaded());
 //}
 
-void UNeuralNetwork::FImplBackEndUEOnly::Run(FOnAsyncRunCompleted& InOutOnAsyncRunCompletedDelegate, const ENeuralNetworkSynchronousMode InSynchronousMode, const ENeuralDeviceType InDeviceType,
+void UNeuralNetwork::FImplBackEndUEOnly::Run(FOnAsyncRunCompleted& InOutOnAsyncRunCompletedDelegate, const ENeuralSynchronousMode InSynchronousMode, const ENeuralDeviceType InDeviceType,
 	const ENeuralDeviceType InInputDeviceType, const ENeuralDeviceType InOutputDeviceType)
 {
 	// Run UNeuralNetwork::UEOnly
@@ -208,7 +208,7 @@ void UNeuralNetwork::FImplBackEndUEOnly::Run(FOnAsyncRunCompleted& InOutOnAsyncR
 					}
 
 					// Broadcast delegates (from the render thread)
-					if (InSynchronousMode == ENeuralNetworkSynchronousMode::Asynchronous)
+					if (InSynchronousMode == ENeuralSynchronousMode::Asynchronous)
 					{
 						bIsBackgroundThreadRunning = true; // This will be done right away
 						GraphBuilder.AddPass(
@@ -227,7 +227,7 @@ void UNeuralNetwork::FImplBackEndUEOnly::Run(FOnAsyncRunCompleted& InOutOnAsyncR
 			);
 
 			// Block thread until GPU has finished
-			if (InSynchronousMode == ENeuralNetworkSynchronousMode::Synchronous)
+			if (InSynchronousMode == ENeuralSynchronousMode::Synchronous)
 			{
 				FNeuralNetworkInferenceUtils::WaitUntilRHIFinished();
 			}
diff --git a/Engine/Plugins/Experimental/NeuralNetworkInference/Source/NeuralNetworkInference/Private/NeuralNetworkImplBackEndUEOnly.h b/Engine/Plugins/Experimental/NeuralNetworkInference/Source/NeuralNetworkInference/Private/NeuralNetworkImplBackEndUEOnly.h
--- a/Engine/Plugins/Experimental/NeuralNetworkInference/Source/NeuralNetworkInference/Private/NeuralNetworkImplBackEndUEOnly.h
+++ b/Engine/Plugins/Experimental/NeuralNetworkInference/Source/NeuralNetworkInference/Private/NeuralNetworkImplBackEndUEOnly.h
@@ -44,7 +44,7 @@ struct UNeuralNetwork::FImplBackEndUEOnly
 	static bool Load(TSharedPtr<FImplBackEndUEOnly>& InOutImplBackEndUEOnly, const TArray<uint8>& InModelReadFromFileInBytes);
 	static bool Load(TSharedPtr<FImplBackEndUEOnly>& InOutImplBackEndUEOnly, FNeuralTensorManager& InTensorManager, const TArray<TSharedPtr<FNeuralOperator>>& InOperators);
 
-	void Run(FOnAsyncRunCompleted& InOutOnAsyncRunCompletedDelegate, const ENeuralNetworkSynchronousMode InSynchronousMode, const ENeuralDeviceType InDeviceType, const ENeuralDeviceType InInputDeviceType, const ENeuralDeviceType InOutputDeviceType);
+	void Run(FOnAsyncRunCompleted& InOutOnAsyncRunCompletedDelegate, const ENeuralSynchronousMode InSynchronousMode, const ENeuralDeviceType InDeviceType, const ENeuralDeviceType InInputDeviceType, const ENeuralDeviceType InOutputDeviceType);
 
 private:
 	/**
diff --git a/Engine/Plugins/Experimental/NeuralNetworkInference/Source/NeuralNetworkInference/Private/NeuralOperators/ConstantOperator.cpp b/Engine/Plugins/Experimental/NeuralNetworkInference/Source/NeuralNetworkInference/Private/NeuralOperators/ConstantOperator.cpp
--- a/Engine/Plugins/Experimental/NeuralNetworkInference/Source/NeuralNetworkInference/Private/NeuralOperators/ConstantOperator.cpp
+++ b/Engine/Plugins/Experimental/NeuralNetworkInference/Source/NeuralNetworkInference/Private/NeuralOperators/ConstantOperator.cpp
@@ -28,7 +28,7 @@ FConstantOperator::FConstantOperator(const FNodeProto* const InNodeProto)
 		if (ValueAttribute->Type == EAttributeProtoAttributeType::TENSOR)
 		{
 			// Get new tensor from TensorProto
-			if (!Tensor.SetFromTensorProto(&ValueAttribute->T, TEXT("FConstantOperator-") + ValueAttribute->Name, ENeuralTensorTypeGPU::Weight))
+			if (!Tensor.SetFromTensorProto(&ValueAttribute->T, TEXT("FConstantOperator-") + ValueAttribute->Name, ENeuralTensorType::Weight))
 			{
 				UE_LOG(LogNeuralNetworkInference, Warning, TEXT("FConstantOperator::FConstantOperator(): SetFromTensorProto() failed."));
 			}
diff --git a/Engine/Plugins/Experimental/NeuralNetworkInference/Source/NeuralNetworkInference/Private/NeuralOperators/ConvTransposeOperator.cpp b/Engine/Plugins/Experimental/NeuralNetworkInference/Source/NeuralNetworkInference/Private/NeuralOperators/ConvTransposeOperator.cpp
--- a/Engine/Plugins/Experimental/NeuralNetworkInference/Source/NeuralNetworkInference/Private/NeuralOperators/ConvTransposeOperator.cpp
+++ b/Engine/Plugins/Experimental/NeuralNetworkInference/Source/NeuralNetworkInference/Private/NeuralOperators/ConvTransposeOperator.cpp
@@ -456,7 +456,7 @@ bool FConvTransposeOperator::SetAndConfigureAuxiliaryTensor()
 		// Initialize XWithZeros
 		XWithZeros.SetNumUninitialized(XWithZerosSizes, X.GetDataType());
 		XWithZeros.SetTo<float>(0.f);
-		XWithZeros.SetTensorTypeGPU(ENeuralTensorTypeGPU::IntermediateInitialized);
+		XWithZeros.SetTensorTypeGPU(ENeuralTensorType::IntermediateInitialized);
 		// Set XSizesForGPU
 		XSizesForGPU.Copy(X.GetSizes());
 	}
diff --git a/Engine/Plugins/Experimental/NeuralNetworkInference/Source/NeuralNetworkInference/Private/NeuralTensor.cpp b/Engine/Plugins/Experimental/NeuralNetworkInference/Source/NeuralNetworkInference/Private/NeuralTensor.cpp
--- a/Engine/Plugins/Experimental/NeuralNetworkInference/Source/NeuralNetworkInference/Private/NeuralTensor.cpp
+++ b/Engine/Plugins/Experimental/NeuralNetworkInference/Source/NeuralNetworkInference/Private/NeuralTensor.cpp
@@ -134,7 +134,7 @@ ENeuralDataType FPrivateNeuralTensor::GetDataTypeFromTensorProtoDataType(const E
 /* FNeuralTensor costructors
  *****************************************************************************/
 
-FNeuralTensor::FNeuralTensor(const ENeuralDataType InDataType, const TArray<int64>& InSizes, const FString& InName, const ENeuralTensorTypeGPU InTensorTypeGPU)
+FNeuralTensor::FNeuralTensor(const ENeuralDataType InDataType, const TArray<int64>& InSizes, const FString& InName, const ENeuralTensorType InTensorTypeGPU)
 	: Name(InName)
 	, DataType(ENeuralDataType::None)
 	, TensorTypeGPU(InTensorTypeGPU)
@@ -144,12 +144,12 @@ FNeuralTensor::FNeuralTensor(const ENeuralDataType InDataType, const TArray<int6
 	SetNumUninitialized(InSizes, InDataType);
 }
 
-FNeuralTensor::FNeuralTensor(const ENeuralDataType InDataType, const int64 InVolume, const FString& InName, const ENeuralTensorTypeGPU InTensorTypeGPU)
+FNeuralTensor::FNeuralTensor(const ENeuralDataType InDataType, const int64 InVolume, const FString& InName, const ENeuralTensorType InTensorTypeGPU)
 	: FNeuralTensor(InDataType, InVolume > 0 ? TArray<int64>({ InVolume }) : TArray<int64>({}), InName, InTensorTypeGPU)
 {
 }
 
-FNeuralTensor::FNeuralTensor(const FString& InName, const ENeuralTensorTypeGPU InTensorTypeGPU)
+FNeuralTensor::FNeuralTensor(const FString& InName, const ENeuralTensorType InTensorTypeGPU)
 	: FNeuralTensor(ENeuralDataType::None, TArray<int64>({}), InName, InTensorTypeGPU)
 {
 }
@@ -159,7 +159,7 @@ FNeuralTensor::FNeuralTensor(const FString& InName, const ENeuralTensorTypeGPU I
 /* FNeuralTensor private costructor
  *****************************************************************************/
 
-FNeuralTensor::FNeuralTensor(const ENeuralDataType InDataType, const void* const InValues, const int64 InSizeOfT, const int64 InValueNum, const TArray<int64>& InSizes, const FString& InName, const ENeuralTensorTypeGPU InTensorTypeGPU)
+FNeuralTensor::FNeuralTensor(const ENeuralDataType InDataType, const void* const InValues, const int64 InSizeOfT, const int64 InValueNum, const TArray<int64>& InSizes, const FString& InName, const ENeuralTensorType InTensorTypeGPU)
 	: FNeuralTensor(InDataType, InSizes, InName, InTensorTypeGPU)
 {
 	// Sanity check
@@ -256,7 +256,7 @@ int64 FNeuralTensor::GetSize(const int32 InDimension) const
 	return (InDimension < GetNumberDimensions() ? Sizes[InDimension] : 1u);
 }
 
-void FNeuralTensor::SetTensorTypeGPU(const ENeuralTensorTypeGPU InTensorTypeGPU)
+void FNeuralTensor::SetTensorTypeGPU(const ENeuralTensorType InTensorTypeGPU)
 {
 	// Sanity check
 	if (PooledBuffer.IsValid() || BufferSRVRef.IsValid() || BufferUAVRef.IsValid())
@@ -281,32 +281,35 @@ void FNeuralTensor::ToGPU_RenderThread(FRDGBuilder* InOutGraphBuilder)
 	// - BUF_Volatile: Updated multiple times in a frame, but does not imply a lifetime of 1 frame. E.g. a vertex buffer you update every frame with new vertices.
 	// - BUF_Transient: Used during 1 frame. Volatile and transient are not mutually exclusive.
 	// - BUF_KeepCPUAccessible: Not needed, I can just copy the final GPU memory back to RAM at the very end
+	// - Input and Intermediate(Not)Initialized currently share the same attributes because input might become intermediate (e.g., if input tensor fed into a ReLU, which simply modifies
+	//   the input FNeuralTensor). However, Intermediate(Not)Initialized and Output do not copy the memory from CPU to GPU but rather simply allocates it.
+	// - Output might also become Intermediate(Not)Initialized (e.g., if Output -> ReLU -> Output), so it is kept as ReadWrite rather than written once to account for this.
 	// Call ToGPU_RenderThread with the right flags
-	if (TensorTypeGPU == ENeuralTensorTypeGPU::Generic)
+	if (TensorTypeGPU == ENeuralTensorType::Generic)
 	{
 		return ToGPU_RenderThread(InOutGraphBuilder, BUF_ShaderResource | BUF_UnorderedAccess, true);
 	}
-	else if (TensorTypeGPU == ENeuralTensorTypeGPU::Input)
+	else if (TensorTypeGPU == ENeuralTensorType::Input)
 	{
 		return ToGPU_RenderThread(InOutGraphBuilder, BUF_ShaderResource | BUF_UnorderedAccess, true);
 	}
-	else if (TensorTypeGPU == ENeuralTensorTypeGPU::IntermediateNotInitialized)
+	else if (TensorTypeGPU == ENeuralTensorType::IntermediateNotInitialized)
 	{
 		return ToGPU_RenderThread(InOutGraphBuilder, BUF_ShaderResource | BUF_UnorderedAccess | BUF_Transient, false);
 	}
-	else if (TensorTypeGPU == ENeuralTensorTypeGPU::IntermediateInitialized)
+	else if (TensorTypeGPU == ENeuralTensorType::IntermediateInitialized)
 	{
 		return ToGPU_RenderThread(InOutGraphBuilder, BUF_ShaderResource | BUF_UnorderedAccess, true);
 	}
-	else if (TensorTypeGPU == ENeuralTensorTypeGPU::Output)
+	else if (TensorTypeGPU == ENeuralTensorType::Output)
 	{
 		return ToGPU_RenderThread(InOutGraphBuilder, BUF_ShaderResource | BUF_UnorderedAccess, false);
 	}
-	else if (TensorTypeGPU == ENeuralTensorTypeGPU::Weight)
+	else if (TensorTypeGPU == ENeuralTensorType::Weight)
 	{
 		return ToGPU_RenderThread(InOutGraphBuilder, BUF_ShaderResource | BUF_Static, true);
 	}
-	UE_LOG(LogNeuralNetworkInference, Warning, TEXT("FNeuralTensor-%s::ToGPU_RenderThread(): Unimplemented TensorTypeGPU = %d. Assuming ENeuralTensorTypeGPU::Generic."), *Name, (int32)TensorTypeGPU);
+	UE_LOG(LogNeuralNetworkInference, Warning, TEXT("FNeuralTensor-%s::ToGPU_RenderThread(): Unimplemented TensorTypeGPU = %d. Assuming ENeuralTensorType::Generic."), *Name, (int32)TensorTypeGPU);
 	return ToGPU_RenderThread(InOutGraphBuilder, BUF_UnorderedAccess, true);
 }
 
@@ -335,7 +338,7 @@ void FNeuralTensor::ToGPU_RenderThread(FRDGBuilder* InOutGraphBuilder, const EBu
 	if (!PooledBuffer)
 	{
 		FRDGBufferDesc BufferDesc;
-		BufferDesc.BytesPerElement = FNeuralDataTypeUtils::GetSize(DataType);
+		BufferDesc.BytesPerElement = FNeuralDataTypeUtils::GetByteSize(DataType);
 		BufferDesc.NumElements = Num();
 		BufferDesc.Usage = InEBufferUsageFlags;
 		BufferDesc.UnderlyingType = FRDGBufferDesc::EUnderlyingType::VertexBuffer;
@@ -441,7 +444,7 @@ bool FNeuralTensor::InitPooledBuffer(void** NativeResource)
 			FRDGBuilder Builder(RHICmdList);
 
 			FRDGBufferDesc BufferDesc;
-			BufferDesc.BytesPerElement = FNeuralDataTypeUtils::GetSize(DataType);
+			BufferDesc.BytesPerElement = FNeuralDataTypeUtils::GetByteSize(DataType);
 			BufferDesc.NumElements = Num();
 			BufferDesc.Usage = BUF_UnorderedAccess | BUF_ShaderResource;
 			BufferDesc.UnderlyingType = FRDGBufferDesc::EUnderlyingType::VertexBuffer;
@@ -527,7 +530,7 @@ void FNeuralTensor::SetNumUninitialized(const TArray<int64>& InSizes, const ENeu
 	if (Sizes.Num() > 0)
 	{
 		Volume = FNeuralNetworkInferenceUtils::Product<int64>(Sizes);
-		const int64 VolumeInBytes = Volume * FNeuralDataTypeUtils::GetSize(DataType);
+		const int64 VolumeInBytes = Volume * FNeuralDataTypeUtils::GetByteSize(DataType);
 		if (VolumeInBytes != ArrayCPU.Num())
 		{
 			ArrayCPU.SetNumUninitialized(VolumeInBytes, bInAllowShrinking); // Pre-allocate TArray
@@ -551,7 +554,7 @@ void FNeuralTensor::SetFromUnderlyingUInt8ArrayCopy(const TArray<uint8>& InArray
 	ArrayCPU = InArray;
 }
 
-bool FNeuralTensor::SetFromTensorProto(const FTensorProto* const InTensorProto, const FString& InTensorName, const ENeuralTensorTypeGPU InTensorTypeGPU)
+bool FNeuralTensor::SetFromTensorProto(const FTensorProto* const InTensorProto, const FString& InTensorName, const ENeuralTensorType InTensorTypeGPU)
 {
 	if (!InTensorProto)
 	{
@@ -586,7 +589,7 @@ bool FNeuralTensor::SetFromTensorProto(const FTensorProto* const InTensorProto,
 		// }
 		// // Read neural tensor from binary data
 		// const FString BinaryWeightFilePath = InExternalDataDirectory / InTensorProto->ExternalData[0].Value;
-		// if (!FModelProtoFileReader::ReadWeightsFromOtxtBinaryFile((char*)GetData(), Num() * FNeuralDataTypeUtils::GetSize(NeuralDataType), BinaryWeightFilePath))
+		// if (!FModelProtoFileReader::ReadWeightsFromOtxtBinaryFile((char*)GetData(), Num() * FNeuralDataTypeUtils::GetByteSize(NeuralDataType), BinaryWeightFilePath))
 		// {
 		// 	UE_LOG(LogNeuralNetworkInference, Warning, TEXT("FNeuralTensor::SetFromTensorProto(): Could not read binary file: %s."), *BinaryWeightFilePath);
 		// 	return false;
@@ -619,7 +622,7 @@ bool FNeuralTensor::Flip(const int32 InDimension)
 	{
 		DimensionOffset *= Sizes[DimensionIndex];
 	}
-	const int64 BytesPerIndex = FNeuralDataTypeUtils::GetSize(DataType);
+	const int64 BytesPerIndex = FNeuralDataTypeUtils::GetByteSize(DataType);
 	const int64 DimensionOffsetInBytes = DimensionOffset * BytesPerIndex;
 
 	// Fill TensorNDIndexes and TensorNDSizes
@@ -703,7 +706,7 @@ bool FNeuralTensor::Transpose()
 			const uint8* const ArrayOnCPUPtr = ArrayCPU.GetData();
 			const int64 Height = Sizes[0];
 			const int64 Width = Sizes[1];
-			const int64 Bytes = FNeuralDataTypeUtils::GetSize(DataType);
+			const int64 Bytes = FNeuralDataTypeUtils::GetByteSize(DataType);
 			for (int64 Y = 0; Y < Height; ++Y)
 			{
 				for (int64 X = 0; X < Width; ++X)
@@ -753,12 +756,12 @@ FString FNeuralTensor::ToString(const int64 InMaxNumberElementsToDisplay, const
 		TensorString += (Name.Len() > 0 ? Name : FString(TEXT("Unnamed FNeuralTensor"))) + TEXT(": ")
 			// DataType
 			+ FNeuralDataTypeUtils::ToString(DataType) + TEXT(", ")
-			+ (TensorTypeGPU == ENeuralTensorTypeGPU::Generic ? TEXT("Generic")
-				: TensorTypeGPU == ENeuralTensorTypeGPU::Input ? TEXT("Input")
-				: TensorTypeGPU == ENeuralTensorTypeGPU::IntermediateNotInitialized ? TEXT("IntermediateNotInitialized")
-				: TensorTypeGPU == ENeuralTensorTypeGPU::IntermediateInitialized ? TEXT("IntermediateInitialized")
-				: TensorTypeGPU == ENeuralTensorTypeGPU::Output ? TEXT("Output")
-				: TensorTypeGPU == ENeuralTensorTypeGPU::Weight ? TEXT("Weight")
+			+ (TensorTypeGPU == ENeuralTensorType::Generic ? TEXT("Generic")
+				: TensorTypeGPU == ENeuralTensorType::Input ? TEXT("Input")
+				: TensorTypeGPU == ENeuralTensorType::IntermediateNotInitialized ? TEXT("IntermediateNotInitialized")
+				: TensorTypeGPU == ENeuralTensorType::IntermediateInitialized ? TEXT("IntermediateInitialized")
+				: TensorTypeGPU == ENeuralTensorType::Output ? TEXT("Output")
+				: TensorTypeGPU == ENeuralTensorType::Weight ? TEXT("Weight")
 				: TEXT("Unknown"))
 			// Volume and sizes
 			+ TEXT(", volume=") + FString::FromInt(Num()) + TEXT(", sizes={");
@@ -880,15 +883,15 @@ void FNeuralTensor::SetFromPointer(const void* const InData, const int64 InSizeO
 
 bool FNeuralTensor::CheckTAndDataTypeResult(const bool bInCheckTAndDataTypeResult, const int64 InSizeOfT) const
 {
-	const int64 SizeOfDataType = FNeuralDataTypeUtils::GetSize(DataType);
+	const int64 ByteSizeOfDataType = FNeuralDataTypeUtils::GetByteSize(DataType);
 	if (!bInCheckTAndDataTypeResult)
 	{
 		const FString DataTypeString = FNeuralDataTypeUtils::ToString(DataType);
 		// sizeof(T) and DataType do not match
-		if (SizeOfDataType != InSizeOfT)
+		if (ByteSizeOfDataType != InSizeOfT)
 		{
 			UE_LOG(LogNeuralNetworkInference, Warning, TEXT("FNeuralTensor-%s::CheckTAndDataTypeResult() failed: DataType = %s, but sizeof(%s) = %d != sizeof(T) = %d."),
-				*Name, *DataTypeString, *DataTypeString, SizeOfDataType, InSizeOfT);
+				*Name, *DataTypeString, *DataTypeString, ByteSizeOfDataType, InSizeOfT);
 		}
 		// sizeof(T) matches, but not the expected DataType
 		else
diff --git a/Engine/Plugins/Experimental/NeuralNetworkInference/Source/NeuralNetworkInference/Private/UEOnly/GraphProtoToNeuralNetworkConverter.cpp b/Engine/Plugins/Experimental/NeuralNetworkInference/Source/NeuralNetworkInference/Private/UEOnly/GraphProtoToNeuralNetworkConverter.cpp
--- a/Engine/Plugins/Experimental/NeuralNetworkInference/Source/NeuralNetworkInference/Private/UEOnly/GraphProtoToNeuralNetworkConverter.cpp
+++ b/Engine/Plugins/Experimental/NeuralNetworkInference/Source/NeuralNetworkInference/Private/UEOnly/GraphProtoToNeuralNetworkConverter.cpp
@@ -80,7 +80,7 @@ bool FGraphProtoToNeuralNetworkConverter::Translate(TArray<TSharedPtr<FNeuralOpe
 			NameIndexMap.Add(GraphProtoInput.Name, Tensors.Num());
 			InputNameIndexMap.Add(GraphProtoInput.Name, Tensors.Num());
 			// Add tensor to TMap
-			Tensors.Emplace(FNeuralTensor(ENeuralDataType::Float, DimensionsAsTArray, GraphProtoInput.Name, ENeuralTensorTypeGPU::Input));
+			Tensors.Emplace(FNeuralTensor(ENeuralDataType::Float, DimensionsAsTArray, GraphProtoInput.Name, ENeuralTensorType::Input));
 		}
 
 		// Read output tensor names (but do not create them yet, it has to happen simultaneously with the operator creation and tensor inlining)
@@ -176,7 +176,7 @@ bool FGraphProtoToNeuralNetworkConverter::CreateOperatorsAndEditTensorArray(TArr
 				// 	FPaths::Split(InModelPath, PathPart, FilenamePart, ExtensionPart);
 				// }
 				// Get NewTensor from TensorProto
-				if (!NewTensor.SetFromTensorProto(&TensorProto, InputTensorName, ENeuralTensorTypeGPU::Weight/*, PathPart*/))
+				if (!NewTensor.SetFromTensorProto(&TensorProto, InputTensorName, ENeuralTensorType::Weight/*, PathPart*/))
 				{
 					UE_LOG(LogNeuralNetworkInference, Warning, TEXT("FGraphProtoToNeuralNetworkConverter::CreateOperatorsAndEditTensorArray(): SetFromTensorProto() failed."));
 					return false;
@@ -224,7 +224,7 @@ bool FGraphProtoToNeuralNetworkConverter::CreateOperatorsAndEditTensorArray(TArr
 				}
 				// Create tensor
 				InOutNameIndexMap.Add(OutputTensorName, TensorIndex);
-				InOutTensors.Push(FNeuralTensor(ENeuralDataType::None, TArray<int64>({}), OutputTensorName, ENeuralTensorTypeGPU::Weight));
+				InOutTensors.Push(FNeuralTensor(ENeuralDataType::None, TArray<int64>({}), OutputTensorName, ENeuralTensorType::Weight));
 
 				// Create constant operator layer and get final tensor value
 				FConstantOperator ConstantOperator(&NodeProto);
@@ -253,7 +253,7 @@ bool FGraphProtoToNeuralNetworkConverter::CreateOperatorsAndEditTensorArray(TArr
 			{
 				UE_LOG(LogNeuralNetworkInference, Display, TEXT("Flipping FConvTranspose weights. This is a temporary measurement but it will keep the weights flipped for now."));
 				// Sanity check
-				if (OperatorInputTensors.Num() > 1 && OperatorInputTensors[1]->GetTensorTypeGPU() != ENeuralTensorTypeGPU::Weight)
+				if (OperatorInputTensors.Num() > 1 && OperatorInputTensors[1]->GetTensorTypeGPU() != ENeuralTensorType::Weight)
 				{
 					UE_LOG(LogNeuralNetworkInference, Warning,
 						TEXT("FGraphProtoToNeuralNetworkConverter::CreateOperatorsAndEditTensorArray(): For now, %s operators require a constant InputTensor[0] (i.e., a fixed weight tensor), it cannot change."), *NodeProto.OperatorType);
@@ -314,7 +314,7 @@ bool FGraphProtoToNeuralNetworkConverter::CreateOperatorsAndEditTensorArray(TArr
 						if (InOutputNameDummyIndexMap.Find(OutputTensorName))
 						{
 							OutputNameIndexMap.Add(OutputTensorName, *ExistingTensorIndex);
-							InOutTensors[*ExistingTensorIndex].SetTensorTypeGPU(ENeuralTensorTypeGPU::Output);
+							InOutTensors[*ExistingTensorIndex].SetTensorTypeGPU(ENeuralTensorType::Output);
 						}
 					}
 					// Subcases 2 and 3: FNeuralTensor not found, it must be defined now (i.e., intermediate tensor of the UNeuralNetworkLegacy)
@@ -323,11 +323,11 @@ bool FGraphProtoToNeuralNetworkConverter::CreateOperatorsAndEditTensorArray(TArr
 					{
 						// Is it also an absolute output of the network?
 						const int64 TensorIndex = InOutTensors.Num();
-						ENeuralTensorTypeGPU TensorTypeGPU = ENeuralTensorTypeGPU::IntermediateNotInitialized;
+						ENeuralTensorType TensorTypeGPU = ENeuralTensorType::IntermediateNotInitialized;
 						if (InOutputNameDummyIndexMap.Find(OutputTensorName))
 						{
 							OutputNameIndexMap.Add(OutputTensorName, TensorIndex);
-							TensorTypeGPU = ENeuralTensorTypeGPU::Output;
+							TensorTypeGPU = ENeuralTensorType::Output;
 						}
 						// Create tensor
 						InOutNameIndexMap.Add(OutputTensorName, TensorIndex);
@@ -352,7 +352,7 @@ bool FGraphProtoToNeuralNetworkConverter::CreateOperatorsAndEditTensorArray(TArr
 						if (InOutputNameDummyIndexMap.Find(OutputTensorName))
 						{
 							OutputNameIndexMap.Add(OutputTensorName, InlinedTensorGlobalIndex);
-							InOutTensors[InlinedTensorGlobalIndex].SetTensorTypeGPU(ENeuralTensorTypeGPU::Output);
+							InOutTensors[InlinedTensorGlobalIndex].SetTensorTypeGPU(ENeuralTensorType::Output);
 						}
 					}
 				}
@@ -396,7 +396,7 @@ bool FGraphProtoToNeuralNetworkConverter::CreateOperatorsAndEditTensorArray(TArr
 						}
 						// Create and push new auxiliary tensor
 						InOutNameIndexMap.Add(AuxiliaryTensorName, InOutTensors.Num());
-						InOutTensors.Push(FNeuralTensor(ENeuralDataType::None, TArray<int64>({}), AuxiliaryTensorName, ENeuralTensorTypeGPU::Generic)); // ENeuralTensorTypeGPU can be modify by FNeuralOperator::Configure()
+						InOutTensors.Push(FNeuralTensor(ENeuralDataType::None, TArray<int64>({}), AuxiliaryTensorName, ENeuralTensorType::Generic)); // ENeuralTensorType can be modify by FNeuralOperator::Configure()
 						AuxiliaryOperatorTensors.Push(&InOutTensors.Last());
 					}
 				}
diff --git a/Engine/Plugins/Experimental/NeuralNetworkInference/Source/NeuralNetworkInference/Public/NeuralEnumClasses.h b/Engine/Plugins/Experimental/NeuralNetworkInference/Source/NeuralNetworkInference/Public/NeuralEnumClasses.h
--- a/Engine/Plugins/Experimental/NeuralNetworkInference/Source/NeuralNetworkInference/Public/NeuralEnumClasses.h
+++ b/Engine/Plugins/Experimental/NeuralNetworkInference/Source/NeuralNetworkInference/Public/NeuralEnumClasses.h
@@ -7,108 +7,151 @@
 #include <type_traits>
 #include "NeuralEnumClasses.generated.h"
 
+/**
+ * It defines the data type (float, double, int32, etc.).
+ */
 UENUM()
 enum class ENeuralDataType : uint8
 {
-	Float,
-	//Double,
-	//Int8,
-	//Int16,
-	Int32,
-	Int64,
-	//UInt8,
-	UInt32,
-	UInt64,
-	None
+	Float, /* 32-bit floating number. */
+	//Double, /* 64-bit floating number. */
+	//Int8, /* 8-bit signed integer. */
+	//Int16, /* 16-bit signed integer. */
+	Int32, /* 32-bit signed integer. */
+	Int64, /* 64-bit signed integer. */
+	//UInt8, /* 8-bit unsigned integer. */
+	UInt32, /* 32-bit unsigned integer. */
+	UInt64, /* 64-bit unsigned integer. */
+	None /* To be used in special cases, e.g., if the data type is unknown yet. */
 };
 
+/**
+ * It defines in which device (CPU, GPU) the desired operation (e.g., the neural network inference) is run.
+ */
 UENUM()
 enum class ENeuralDeviceType : uint8
 {
-	CPU,
-	GPU,
-	None
+	CPU, /* The operation will occur on the CPU. */
+	GPU /* The operation will occur on the GPU. */
 };
 
 /**
- * Whether UNeuralNetwork::Run() will block the thread until completed (Synchronous), or whether it will run on a background thread,
- * not blocking the calling thread (Asynchronous).
+ * Whether the operation (e.g., UNeuralNetwork::Run()) will run in the calling thread (Synchronous) or in a background thread (Asynchronous).
  */
 UENUM()
-enum class ENeuralNetworkSynchronousMode : uint8
+enum class ENeuralSynchronousMode : uint8
 {
-	Synchronous, /* UNeuralNetwork::Run() will block the thread until the network evaluation (i.e., forward pass) has finished. */
 	/**
-	 * UNeuralNetwork::Run() will initialize a forward pass request on a background thread, not blocking the thread that called it.
-	 * The user should register to UNeuralNetwork's delegate to know when the forward pass has finished.
+	 * Safer and simpler to use.
+	 * The operation will run in the calling thread thus blocking that thread until completed.
+	 */
+	Synchronous,
+	/**
+	 * More complex but potentially more efficient.
+	 * The operation will initialize a compute request on a background thread and return before its completion, not blocking the calling thread that
+	 * called it.
+	 * The user should register to a delegate which will notify them when the operation has finished asynchronously
+	 * (e.g., UNeuralNetwork::GetOnAsyncRunCompletedDelegate() for asynchronous UNeuralNetwork::Run()).
 	 *
-	 * Very important: It takes ~1 millisecond to start the background thread. If your network runs synchronously faster than 1 msec,
-	 * using asynchronous running will make the game (main) thread slower than running it synchronously.
+	 * Very important: It takes ~1 millisecond to start the background thread. If your operation (e.g., your network inference time) runs
+	 * synchronously faster than 1-2 msec, using asynchronous is not recommended because it will slow down both the main and bakground threads.
 	 */
 	Asynchronous
 };
 
+/**
+ * The type of the neural tensor. E.g., a Weight tensor will be read-only and never modified, an Input tensor will be modified by the user, an
+ * IntermediateNotInitialized tensor will change on each frame, etc.
+ *
+ * Although conceptually this could apply to both the CPU and GPU versions, in practice only the GPU performance is affected by this setting so far.
+ */
 UENUM()
-enum class ENeuralNetworkDelegateThreadMode : uint8
+enum class ENeuralTensorType : uint8
 {
-	GameThread, /* Recommended and default value. The UNeuralNetwork delegate will be called from the game thread. */
 	/**
-	 * Not recommended, use at your own risk.
-	 * The UNeuralNetwork delegate could be called from any thread.
-	 * Running UClass functions from background threads is not safe (e.g., it might crash if the editor is closed while accessing UNeuralNetwork information).
-	 * Thus "AnyThread" is only safe if you have guarantees that the program will not be terminated while calling UNeuralNetwork functions.
+	 * Safe and generic tensor that works in every situation (e.g., ReadWrite, not volatile). However, it might not be the most efficient one for
+	 * most cases.
 	 */
-	AnyThread
+	Generic,
+	/**
+	 * Input tensor of a neural network, usually copied from the CPU and usually ReadOnly (but might be modified by in-place operators like ReLU).
+	 */
+	Input,
+	/**
+	 * Intermediate tensor of a neural network (output of at least a layer and input of at least some other layer). Not copied from CPU, ReadWrite,
+	 * and transient.
+	 */
+	IntermediateNotInitialized,
+	/**
+	 * Intermediate tensor that is initialized with CPU data (e.g., XWithZeros in FConvTranpose). Copied from CPU.
+	 */
+	IntermediateInitialized,
+	/**
+	 * Output tensor of a neural network. Not copied from CPU and ReadWrite.
+	 */
+	Output,
+	/**
+	 * Weights of a particular operator/layer. Copied from CPU, ReadOnly, and initialized from CPU memory.
+	 */
+	Weight
 };
 
-
 /**
- * Although conceptually this could apply to both the CPU and GPU versions, in practice only the GPU performance is affected by this setting.
- * Input and Intermediate(Not)Initialized currently share the same attributes because input might become intermediate (e.g., if input tensor fed into a ReLU, which simply modifies
- * the input FNeuralTensor). However, Intermediate(Not)Initialized and Output do not copy the memory from CPU to GPU but rather simply allocates it.
- * Output might also become Intermediate(Not)Initialized (e.g., if Output -> ReLU -> Output), so it is kept as ReadWrite rather than written once to account for this.
+ * After an asynchronous operation has finished, whether the callback functions tied to the delegate will be called from the game/main thread (highly
+ * recommended) or from any thread (not fully Unreal safe).
+ * This enum class is only useful for the case of ENeuralSynchronousMode::Asynchronous.
  */
 UENUM()
-enum class ENeuralTensorTypeGPU : uint8
+enum class ENeuralThreadMode : uint8
 {
-	Generic,					/** Generic tensor that works in every situation (ReadWrite), although it might not be the most efficient one. */
-	Input,						/** Input tensor of the UNeuralNetworkLegacy. Copied from CPU and ReadWrite (but usually ReadOnly). */
-	IntermediateNotInitialized,	/** Intermediate tensor of the UNeuralNetworkLegacy (output of at least a layer and input of at least some other layer). Not copied from CPU, ReadWrite, and transient. */
-	IntermediateInitialized,	/** Intermediate tensor that is initialized with CPU data (e.g., XWithZeros in FConvTranpose). Copied from CPU. */
-	Output,						/** Output tensor of the UNeuralNetworkLegacy. Not copied from CPU and ReadWrite. */
-	Weight						/** Weights of a particular operator/layer. Copied from CPU, ReadOnly, and initialized from CPU memory. */
+	/**
+	 * Highly recommended and default value. The callback functions tied to the delegate will be called from the game/main thread.
+	 */
+	GameThread,
+	/**
+	 * Not recommended, use at your own risk (potentially more efficient than GameThread but not UE-safe in all cases).
+	 * The callback functions tied to the delegate will be called from the same thread at which the operation was asynchronously computed, thus this
+	 * could happen from any thread.
+	 *
+	 * Main issue: U-classes (e.g., UNeuralNetwork) should not be used from non-game threads because it is not safe and could lead to issues/crashes
+	 * if not properly handled by the user.
+	 * E.g., the GC does not work from non-game threads so it might crash if the editor/UE is closed while accessing UNeuralNetwork information from
+	 * the callback function.
+	 *
+	 * Thus "AnyThread" is only safe if you have guarantees that your case is fully U-class safe. E.g., the program will not be terminated while
+	 * calling the callback function.
+	 */
+	AnyThread
 };
 
 
 
 /**
- * Auxiliary utils class for ENeuralDataType
+ * Auxiliary class consisting of static and auxiliary functions for ENeuralDataType.
  */
-
 class NEURALNETWORKINFERENCE_API FNeuralDataTypeUtils
 {
 public:
-	static FString ToString(const ENeuralDataType InDataType);
-	static int64 GetSize(const ENeuralDataType InDataType);
-	static EPixelFormat GetPixelFormat(const ENeuralDataType InDataType);
-
 	/**
-	 * It checks whether T and InDataType are the same type. E.g.,
-	 * checkf(CheckTAndDataType<float>(), TEXT("Expected a ENeuralDataType::Float type."));
+	 * It returns the byte size of the input ENeuralDataType (e.g., 4 for ENeuralDataType::Float or ENeuralDataType::Int32).
 	 */
-	template <typename T>
-	static bool CheckTAndDataType(const ENeuralDataType InDataType)
-	{
-		return InDataType == GetDataType<T>();
-	}
+	static int64 GetByteSize(const ENeuralDataType InDataType);
 
 	/**
-	 * It gets the data type from the type T. E.g.,
-	 * checkf(InDataType == GetDataType<float>(), TEXT("InDataType == GetDataType<float>() failed!"))
-	 * FNeuralTensor(FNeuralDataTypeUtils::GetDataType<T>(), InArray.GetData(), ...)
+	 * It returns the data type from the type T (e.g., ENeuralDataType::Float for GetDataType<float>()).
 	 */
 	template <typename T>
 	static ENeuralDataType GetDataType();
+
+	/**
+	 * It returns the pixel format of the input ENeuralDataType (e.g., EPixelFormat::PF_R32_FLOAT for ENeuralDataType::Float or "PF_R32_SINT" for Int32).
+	 */
+	static EPixelFormat GetPixelFormat(const ENeuralDataType InDataType);
+
+	/**
+	 * It returns the FString name of the input ENeuralDataType (e.g., "Float" for ENeuralDataType::Float).
+	 */
+	static FString ToString(const ENeuralDataType InDataType);
 };
 
 
diff --git a/Engine/Plugins/Experimental/NeuralNetworkInference/Source/NeuralNetworkInference/Public/NeuralNetwork.h b/Engine/Plugins/Experimental/NeuralNetworkInference/Source/NeuralNetworkInference/Public/NeuralNetwork.h
--- a/Engine/Plugins/Experimental/NeuralNetworkInference/Source/NeuralNetworkInference/Public/NeuralNetwork.h
+++ b/Engine/Plugins/Experimental/NeuralNetworkInference/Source/NeuralNetworkInference/Public/NeuralNetwork.h
@@ -97,19 +97,19 @@ public:
 	 * Getter and setter functions for SynchronousMode.
 	 * @see SynchronousMode and GetOnAsyncRunCompletedDelegate() for more details.
 	 */
-	ENeuralNetworkSynchronousMode GetSynchronousMode() const;
-	void SetSynchronousMode(const ENeuralNetworkSynchronousMode InSynchronousMode);
+	ENeuralSynchronousMode GetSynchronousMode() const;
+	void SetSynchronousMode(const ENeuralSynchronousMode InSynchronousMode);
 
 	/**
 	 * GetOnAsyncRunCompletedDelegate() returns a FOnAsyncRunCompleted delegate that will be called when async UNeuralNetwork::Run() is completed.
-	 * - If SynchronousMode == ENeuralNetworkSynchronousMode::Asynchronous, the FOnAsyncRunCompleted delegate could be triggered from any thread.
-	 * - If SynchronousMode == ENeuralNetworkSynchronousMode::Synchronous, UNeuralNetwork::Run() will block the calling thread until completed, so a
+	 * - If SynchronousMode == ENeuralSynchronousMode::Asynchronous, the FOnAsyncRunCompleted delegate could be triggered from any thread.
+	 * - If SynchronousMode == ENeuralSynchronousMode::Synchronous, UNeuralNetwork::Run() will block the calling thread until completed, so a
 	 *   callback delegate is not required.
 	 */
 	DECLARE_DELEGATE(FOnAsyncRunCompleted);
 	FOnAsyncRunCompleted& GetOnAsyncRunCompletedDelegate();
-	ENeuralNetworkDelegateThreadMode GetOnAsyncRunCompletedDelegateMode() const;
-	void SetOnAsyncRunCompletedDelegateMode(const ENeuralNetworkDelegateThreadMode InDelegateThreadMode);
+	ENeuralThreadMode GetOnAsyncRunCompletedDelegateMode() const;
+	void SetOnAsyncRunCompletedDelegateMode(const ENeuralThreadMode InDelegateThreadMode);
 	
 	/**
 	 * Whether GPU execution is supported for this platform. It will return:
@@ -201,12 +201,12 @@ protected:
 	 * background thread, not blocking the calling thread (Asynchronous).
 	 * If asynchronous, DelegateThreadMode will define whether the callback delegate is called from the game thread (highly recommended) or from
 	 * any available thread (not fully thread safe).
-	 * @see ENeuralNetworkSynchronousMode, ENeuralNetworkDelegateThreadMode for more details.
+	 * @see ENeuralSynchronousMode, ENeuralThreadMode for more details.
 	 */
 	UPROPERTY(Transient, VisibleAnywhere, Category = "Neural Network Inference")
-	ENeuralNetworkSynchronousMode SynchronousMode;
+	ENeuralSynchronousMode SynchronousMode;
 	UPROPERTY(Transient, VisibleAnywhere, Category = "Neural Network Inference")
-	ENeuralNetworkDelegateThreadMode DelegateThreadMode;
+	ENeuralThreadMode DelegateThreadMode;
 
 	/**
 	 * Original model file path from which this UNeuralNetwork was loaded from.
diff --git a/Engine/Plugins/Experimental/NeuralNetworkInference/Source/NeuralNetworkInference/Public/NeuralTensor.h b/Engine/Plugins/Experimental/NeuralNetworkInference/Source/NeuralNetworkInference/Public/NeuralTensor.h
--- a/Engine/Plugins/Experimental/NeuralNetworkInference/Source/NeuralNetworkInference/Public/NeuralTensor.h
+++ b/Engine/Plugins/Experimental/NeuralNetworkInference/Source/NeuralNetworkInference/Public/NeuralTensor.h
@@ -34,9 +34,9 @@ public:
 	 * @param InSizes Set to empty (or omit argument) if memory allocation is not required or the final size is unknown.
 	 * @param InName Used for GPU debugging and the ToString() function.
 	*/
-	explicit FNeuralTensor(const ENeuralDataType InDataType = ENeuralDataType::None, const TArray<int64>& InSizes = TArray<int64>(), const FString& InName = TEXT("FNeuralTensor"), const ENeuralTensorTypeGPU InTensorTypeGPU = ENeuralTensorTypeGPU::Generic);
-	explicit FNeuralTensor(const ENeuralDataType InDataType, const int64 InVolume, const FString& InName = TEXT("FNeuralTensor"), const ENeuralTensorTypeGPU InTensorTypeGPU = ENeuralTensorTypeGPU::Generic);
-	explicit FNeuralTensor(const FString& InName, const ENeuralTensorTypeGPU InTensorTypeGPU = ENeuralTensorTypeGPU::Generic);
+	explicit FNeuralTensor(const ENeuralDataType InDataType = ENeuralDataType::None, const TArray<int64>& InSizes = TArray<int64>(), const FString& InName = TEXT("FNeuralTensor"), const ENeuralTensorType InTensorTypeGPU = ENeuralTensorType::Generic);
+	explicit FNeuralTensor(const ENeuralDataType InDataType, const int64 InVolume, const FString& InName = TEXT("FNeuralTensor"), const ENeuralTensorType InTensorTypeGPU = ENeuralTensorType::Generic);
+	explicit FNeuralTensor(const FString& InName, const ENeuralTensorType InTensorTypeGPU = ENeuralTensorType::Generic);
 
 	/**
 	 * Performance-wise, this constructor makes a deep copy of the data (not optimal). For maximum speed, use the other constructors + GetData()/GetDataCasted().
@@ -47,7 +47,7 @@ public:
 	 * @param InName Used for GPU debugging and the ToString() function.
 	 */
 	template <typename T>
-	explicit FNeuralTensor(const TArray<T>& InArray, const TArray<int64>& InSizes = TArray<int64>(), const FString& InName = TEXT("FNeuralTensor"), const ENeuralTensorTypeGPU InTensorTypeGPU = ENeuralTensorTypeGPU::Generic);
+	explicit FNeuralTensor(const TArray<T>& InArray, const TArray<int64>& InSizes = TArray<int64>(), const FString& InName = TEXT("FNeuralTensor"), const ENeuralTensorType InTensorTypeGPU = ENeuralTensorType::Generic);
 
 	/**
 	 * Copy constructor.
@@ -72,13 +72,13 @@ public:
 
 	/**
 	 * Comparison operator (equal). Returns true if the dimensions, sizes, scalar type, and data match with each other.
-	 * It does not consider other properties of the FNeuralTensor (such as ENeuralTensorTypeGPU).
+	 * It does not consider other properties of the FNeuralTensor (such as ENeuralTensorType).
 	 */
 	bool operator==(const FNeuralTensor& InTensorToCopy) const;
 
 	/**
 	 * Comparison operator (not equal). Returns the opposite than the equal (==) operator.
-	 * It does not consider other properties of the FNeuralTensor (such as ENeuralTensorTypeGPU).
+	 * It does not consider other properties of the FNeuralTensor (such as ENeuralTensorType).
 	 */
 	FORCEINLINE bool operator!=(const FNeuralTensor& InTensorToCopy) const;
 
@@ -157,15 +157,15 @@ public:
 	FORCEINLINE bool IsEmpty() const;
 
 	/**
-	 * It returns the ENeuralTensorTypeGPU (Generic, Input, Intermediate(Not)Initialized, Output, Weight, etc.).
+	 * It returns the ENeuralTensorType (Generic, Input, Intermediate(Not)Initialized, Output, Weight, etc.).
 	 */
-	FORCEINLINE ENeuralTensorTypeGPU GetTensorTypeGPU() const;
+	FORCEINLINE ENeuralTensorType GetTensorTypeGPU() const;
 
 	/**
-	 * It sets the ENeuralTensorTypeGPU (Generic, Input, Intermediate(Not)Initialized, Output, Weight, etc.).
+	 * It sets the ENeuralTensorType (Generic, Input, Intermediate(Not)Initialized, Output, Weight, etc.).
 	 * If the GPU memory was already initialized, it will also UE_LOG a warning.
 	 */
-	void SetTensorTypeGPU(const ENeuralTensorTypeGPU InTensorTypeGPU);
+	void SetTensorTypeGPU(const ENeuralTensorType InTensorTypeGPU);
 
 	FORCEINLINE void SetEnableGPU(const bool bInEnableGPU);
 
@@ -216,7 +216,7 @@ public:
 	 * It will fill the current neural tensor with the input FTensorProto.
 	 * @returns Whether the conversion was successful.
 	 */
-	bool SetFromTensorProto(const FTensorProto* const InTensorProto, const FString& InTensorName, const ENeuralTensorTypeGPU InTensorTypeGPU);
+	bool SetFromTensorProto(const FTensorProto* const InTensorProto, const FString& InTensorName, const ENeuralTensorType InTensorTypeGPU);
 
 	/**
 	 * It sets all the elements of the FNeuralTensor to InValue.
@@ -269,7 +269,7 @@ protected:
 	UPROPERTY(VisibleAnywhere, Category = "Neural Network Inference")
 	TArray<int64> Sizes;
 	UPROPERTY(VisibleAnywhere, Category = "Neural Network Inference")
-	ENeuralTensorTypeGPU TensorTypeGPU;
+	ENeuralTensorType TensorTypeGPU;
 
 private:
 	/** CPU-based variables */
@@ -290,7 +290,7 @@ private:
 	 * Auxiliary for the templated constructor FNeuralTensor(const ENeuralDataType InDataType, const TArray<T>& InArray, ...).
 	 * InSizeOfT and sizeof(InDataType) should match, as well as the volume of InValueNum and InSizes.
 	 */
-	explicit FNeuralTensor(const ENeuralDataType InDataType, const void* const InValues, const int64 InSizeOfT, const int64 InValueNum, const TArray<int64>& InSizes, const FString& InName, const ENeuralTensorTypeGPU InTensorTypeGPU);
+	explicit FNeuralTensor(const ENeuralDataType InDataType, const void* const InValues, const int64 InSizeOfT, const int64 InValueNum, const TArray<int64>& InSizes, const FString& InName, const ENeuralTensorType InTensorTypeGPU);
 
 	/**
 	 * Auxiliary for the templated SetFromArrayCopy().
@@ -301,7 +301,7 @@ private:
 	 * Checks and warns whether the current data type T is incompatible with DataType.
 	 */
 	template<typename T>
-	bool CheckTAndDataType() const;
+	bool CheckTAndDataTypeEquivalent() const;
 	bool CheckTAndDataTypeResult(const bool bInCheckTAndDataTypeResult, const int64 InSizeOfT) const;
 };
 
@@ -311,7 +311,7 @@ private:
  *****************************************************************************/
 
 template <typename T>
-FNeuralTensor::FNeuralTensor(const TArray<T>& InArray, const TArray<int64>& InSizes, const FString& InName, const ENeuralTensorTypeGPU InTensorTypeGPU)
+FNeuralTensor::FNeuralTensor(const TArray<T>& InArray, const TArray<int64>& InSizes, const FString& InName, const ENeuralTensorType InTensorTypeGPU)
 	: FNeuralTensor(FNeuralDataTypeUtils::GetDataType<T>(), InArray.GetData(), sizeof(T), InArray.Num(), (InSizes.Num() > 0 ? InSizes : TArray<int64>({ (int64)InArray.Num() })), InName, InTensorTypeGPU)
 {}
 
@@ -323,14 +323,14 @@ bool FNeuralTensor::operator!=(const FNeuralTensor& InTensorToCopy) const
 template <typename T, typename TInput>
 T& FNeuralTensor::At(const TInput InIndex)
 {
-	checkf(CheckTAndDataType<T>(), TEXT("FNeuralTensor::At(): CheckTAndDataType failed."));
+	checkf(CheckTAndDataTypeEquivalent<T>(), TEXT("FNeuralTensor::At(): CheckTAndDataType failed."));
 	return ((T*)ArrayCPU.GetData())[InIndex];
 }
 
 template <typename T, typename TInput>
 const T& FNeuralTensor::At(const TInput InIndex) const
 {
-	checkf(CheckTAndDataType<T>(), TEXT("FNeuralTensor::At(): CheckTAndDataType failed."));
+	checkf(CheckTAndDataTypeEquivalent<T>(), TEXT("FNeuralTensor::At(): CheckTAndDataType failed."));
 	return ((T*)ArrayCPU.GetData())[InIndex];
 }
 
@@ -338,7 +338,7 @@ template <typename T>
 TArray<T> FNeuralTensor::GetArrayCopy() const
 {
 	TArray<T> Array;
-	if (CheckTAndDataType<T>())
+	if (CheckTAndDataTypeEquivalent<T>())
 	{
 		Array.SetNumUninitialized(Num());
 		FMemory::Memcpy(Array.GetData(), ArrayCPU.GetData(), NumInBytes());
@@ -364,14 +364,14 @@ const void* const FNeuralTensor::GetData() const
 template<typename T>
 T* FNeuralTensor::GetDataCasted()
 {
-	checkf(CheckTAndDataType<T>(), TEXT("FNeuralTensor::GetDataCasted(): CheckTAndDataType failed."));
+	checkf(CheckTAndDataTypeEquivalent<T>(), TEXT("FNeuralTensor::GetDataCasted(): CheckTAndDataType failed."));
 	return (T*)ArrayCPU.GetData();
 }
 
 template<typename T>
 const T* const FNeuralTensor::GetDataCasted() const
 {
-	checkf(CheckTAndDataType<T>(), TEXT("FNeuralTensor::GetDataCasted(): CheckTAndDataType failed."));
+	checkf(CheckTAndDataTypeEquivalent<T>(), TEXT("FNeuralTensor::GetDataCasted(): CheckTAndDataType failed."));
 	return (T*)ArrayCPU.GetData();
 }
 
@@ -415,7 +415,7 @@ bool FNeuralTensor::IsEmpty() const
 	return ArrayCPU.IsEmpty();
 }
 
-ENeuralTensorTypeGPU FNeuralTensor::GetTensorTypeGPU() const
+ENeuralTensorType FNeuralTensor::GetTensorTypeGPU() const
 {
 	return TensorTypeGPU;
 }
@@ -433,7 +433,7 @@ void FNeuralTensor::SetNumUninitialized(const int64 InVolume, const ENeuralDataT
 template<typename T>
 void FNeuralTensor::SetFromArrayCopy(const TArray<T>& InArray)
 {
-	if (CheckTAndDataType<T>())
+	if (CheckTAndDataTypeEquivalent<T>())
 	{
 		SetFromPointer(InArray.GetData(), sizeof(T), InArray.Num());
 	}
@@ -442,14 +442,14 @@ void FNeuralTensor::SetFromArrayCopy(const TArray<T>& InArray)
 template<typename T, typename TInput>
 void FNeuralTensor::SetTo(const TInput InValue)
 {
-	if (CheckTAndDataType<T>())
+	if (CheckTAndDataTypeEquivalent<T>())
 	{
 		std::fill(GetDataCasted<T>(), GetDataCasted<T>() + Num(), (T)InValue); // There is no UE equivalent for std::fill
 	}
 }
 
 template<typename T>
-bool FNeuralTensor::CheckTAndDataType() const
+bool FNeuralTensor::CheckTAndDataTypeEquivalent() const
 {
-	return CheckTAndDataTypeResult(FNeuralDataTypeUtils::CheckTAndDataType<T>(DataType), sizeof(T));
+	return CheckTAndDataTypeResult(FNeuralDataTypeUtils::GetDataType<T>() == DataType, sizeof(T));
 }
diff --git a/Engine/Plugins/Experimental/NeuralNetworkInference/Source/NeuralNetworkInferenceQA/Private/ModelUnitTester.cpp b/Engine/Plugins/Experimental/NeuralNetworkInference/Source/NeuralNetworkInferenceQA/Private/ModelUnitTester.cpp
--- a/Engine/Plugins/Experimental/NeuralNetworkInference/Source/NeuralNetworkInferenceQA/Private/ModelUnitTester.cpp
+++ b/Engine/Plugins/Experimental/NeuralNetworkInference/Source/NeuralNetworkInferenceQA/Private/ModelUnitTester.cpp
@@ -169,13 +169,13 @@ bool FModelUnitTester::ModelLoadAccuracyAndSpeedTests(const FString& InProjectCo
 		const bool bShouldRunUEAndORTBackEnd = (InCPURepetitionsForUEAndORTBackEnd[ModelIndex] + InGPURepetitionsForUEAndORTBackEnd[ModelIndex] > 0);
 		if (bShouldRunUEAndORTBackEnd)
 		{
-			bDidGlobalTestPassed &= ModelAccuracyTest(Network, ENeuralNetworkSynchronousMode::Synchronous, UNeuralNetwork::ENeuralBackEnd::UEAndORT, InInputArrayValues, CPUGroundTruths, GPUGroundTruths);
-			bDidGlobalTestPassed &= ModelAccuracyTest(Network, ENeuralNetworkSynchronousMode::Asynchronous, UNeuralNetwork::ENeuralBackEnd::UEAndORT, InInputArrayValues, CPUGroundTruths, GPUGroundTruths);
+			bDidGlobalTestPassed &= ModelAccuracyTest(Network, ENeuralSynchronousMode::Synchronous, UNeuralNetwork::ENeuralBackEnd::UEAndORT, InInputArrayValues, CPUGroundTruths, GPUGroundTruths);
+			bDidGlobalTestPassed &= ModelAccuracyTest(Network, ENeuralSynchronousMode::Asynchronous, UNeuralNetwork::ENeuralBackEnd::UEAndORT, InInputArrayValues, CPUGroundTruths, GPUGroundTruths);
 		}
 		const bool bShouldRunUEOnlyBackEnd = false; // (InCPURepetitionsForUEOnlyBackEnd[ModelIndex] * InCPURepetitionsForUEOnlyBackEnd[ModelIndex] > 0);
 		if (bShouldRunUEOnlyBackEnd)
 		{
-			bDidGlobalTestPassed &= ModelAccuracyTest(Network, ENeuralNetworkSynchronousMode::Synchronous, UNeuralNetwork::ENeuralBackEnd::UEOnly, InInputArrayValues, CPUGroundTruths, GPUGroundTruths);
+			bDidGlobalTestPassed &= ModelAccuracyTest(Network, ENeuralSynchronousMode::Synchronous, UNeuralNetwork::ENeuralBackEnd::UEOnly, InInputArrayValues, CPUGroundTruths, GPUGroundTruths);
 		}
 
 		UE_LOG(LogNeuralNetworkInferenceQA, Display, TEXT("---------------------------------------------------------------------------------------------------------------------------------"));
@@ -189,12 +189,12 @@ bool FModelUnitTester::ModelLoadAccuracyAndSpeedTests(const FString& InProjectCo
 			UE_LOG(LogNeuralNetworkInferenceQA, Display, TEXT("-------------------- %s - Network %s Load and Run - %s"), *ModelName, *ModelType, *ModelFilePath);
 			if (bShouldRunUEAndORTBackEnd)
 			{
-				bDidGlobalTestPassed &= ModelAccuracyTest(NetworkONNXOrORTLoadTest(ModelFilePath), ENeuralNetworkSynchronousMode::Synchronous, UNeuralNetwork::ENeuralBackEnd::UEAndORT, InInputArrayValues, CPUGroundTruths, GPUGroundTruths);
-				bDidGlobalTestPassed &= ModelAccuracyTest(NetworkONNXOrORTLoadTest(ModelFilePath), ENeuralNetworkSynchronousMode::Asynchronous, UNeuralNetwork::ENeuralBackEnd::UEAndORT, InInputArrayValues, CPUGroundTruths, GPUGroundTruths);
+				bDidGlobalTestPassed &= ModelAccuracyTest(NetworkONNXOrORTLoadTest(ModelFilePath), ENeuralSynchronousMode::Synchronous, UNeuralNetwork::ENeuralBackEnd::UEAndORT, InInputArrayValues, CPUGroundTruths, GPUGroundTruths);
+				bDidGlobalTestPassed &= ModelAccuracyTest(NetworkONNXOrORTLoadTest(ModelFilePath), ENeuralSynchronousMode::Asynchronous, UNeuralNetwork::ENeuralBackEnd::UEAndORT, InInputArrayValues, CPUGroundTruths, GPUGroundTruths);
 			}
 			if (bShouldRunUEOnlyBackEnd)
 			{
-				bDidGlobalTestPassed &= ModelAccuracyTest(NetworkONNXOrORTLoadTest(ModelFilePath), ENeuralNetworkSynchronousMode::Synchronous, UNeuralNetwork::ENeuralBackEnd::UEOnly, InInputArrayValues, CPUGroundTruths, GPUGroundTruths);
+				bDidGlobalTestPassed &= ModelAccuracyTest(NetworkONNXOrORTLoadTest(ModelFilePath), ENeuralSynchronousMode::Synchronous, UNeuralNetwork::ENeuralBackEnd::UEOnly, InInputArrayValues, CPUGroundTruths, GPUGroundTruths);
 			}
 		}
 #else //WITH_EDITOR
@@ -289,7 +289,7 @@ UNeuralNetwork* FModelUnitTester::NetworkONNXOrORTLoadTest(const FString& InMode
 	return Network;
 }
 
-bool FModelUnitTester::ModelAccuracyTest(UNeuralNetwork* InOutNetwork, const ENeuralNetworkSynchronousMode InSynchronousMode, const UNeuralNetwork::ENeuralBackEnd InBackEnd, const TArray<float>& InInputArrayValues,
+bool FModelUnitTester::ModelAccuracyTest(UNeuralNetwork* InOutNetwork, const ENeuralSynchronousMode InSynchronousMode, const UNeuralNetwork::ENeuralBackEnd InBackEnd, const TArray<float>& InInputArrayValues,
 	const TArray<double>& InCPUGroundTruths, const TArray<double>& InGPUGroundTruths)
 {
 	// Sanity check
@@ -320,13 +320,13 @@ bool FModelUnitTester::ModelAccuracyTest(UNeuralNetwork* InOutNetwork, const ENe
 	const ENeuralDeviceType OriginalDeviceType = InOutNetwork->GetDeviceType();
 	const ENeuralDeviceType OriginalInputDeviceType = InOutNetwork->GetInputDeviceType();
 	const ENeuralDeviceType OriginalOutputDeviceType = InOutNetwork->GetOutputDeviceType();
-	const ENeuralNetworkSynchronousMode OriginalSynchronousMode = InOutNetwork->GetSynchronousMode();
-	const ENeuralNetworkDelegateThreadMode OriginalDelegateThreadMode = InOutNetwork->GetOnAsyncRunCompletedDelegateMode();
+	const ENeuralSynchronousMode OriginalSynchronousMode = InOutNetwork->GetSynchronousMode();
+	const ENeuralThreadMode OriginalDelegateThreadMode = InOutNetwork->GetOnAsyncRunCompletedDelegateMode();
 	const UNeuralNetwork::ENeuralBackEnd OriginalBackEnd = InOutNetwork->GetBackEnd();
 	
 	// Set (a)synchronous Mode
 	InOutNetwork->SetSynchronousMode(InSynchronousMode);
-	InOutNetwork->SetOnAsyncRunCompletedDelegateMode(ENeuralNetworkDelegateThreadMode::AnyThread);
+	InOutNetwork->SetOnAsyncRunCompletedDelegateMode(ENeuralThreadMode::AnyThread);
 
 	// Set back end
 	if (!InOutNetwork->SetBackEnd(InBackEnd))
@@ -424,8 +424,8 @@ void FModelUnitTester::ModelAccuracyTestRun(TArray<TArray<float>>& OutOutputs, U
 	const ENeuralDeviceType InDeviceType, const ENeuralDeviceType InInputDeviceType, const ENeuralDeviceType InOutputDeviceType)
 {
 	std::atomic<bool> bDidAsyncNetworkFinished(false);
-	const ENeuralNetworkSynchronousMode SynchronousMode = InOutNetwork->GetSynchronousMode();
-	if (SynchronousMode == ENeuralNetworkSynchronousMode::Asynchronous)
+	const ENeuralSynchronousMode SynchronousMode = InOutNetwork->GetSynchronousMode();
+	if (SynchronousMode == ENeuralSynchronousMode::Asynchronous)
 	{
 		InOutNetwork->GetOnAsyncRunCompletedDelegate().BindLambda([&bDidAsyncNetworkFinished]()
 		{
@@ -443,7 +443,7 @@ void FModelUnitTester::ModelAccuracyTestRun(TArray<TArray<float>>& OutOutputs, U
 		}
 		InOutNetwork->Run();
 		// If async test, wait until inference is completed
-		if (SynchronousMode == ENeuralNetworkSynchronousMode::Asynchronous)
+		if (SynchronousMode == ENeuralSynchronousMode::Asynchronous)
 		{
 			while (!bDidAsyncNetworkFinished)
 			{
@@ -456,7 +456,7 @@ void FModelUnitTester::ModelAccuracyTestRun(TArray<TArray<float>>& OutOutputs, U
 		}
 		OutOutputs.Emplace(InOutNetwork->GetOutputTensor().GetArrayCopy<float>());
 	}
-	if (SynchronousMode == ENeuralNetworkSynchronousMode::Asynchronous)
+	if (SynchronousMode == ENeuralSynchronousMode::Asynchronous)
 	{
 		InOutNetwork->GetOnAsyncRunCompletedDelegate().Unbind();
 	}
diff --git a/Engine/Plugins/Experimental/NeuralNetworkInference/Source/NeuralNetworkInferenceQA/Private/ModelUnitTester.h b/Engine/Plugins/Experimental/NeuralNetworkInference/Source/NeuralNetworkInferenceQA/Private/ModelUnitTester.h
--- a/Engine/Plugins/Experimental/NeuralNetworkInference/Source/NeuralNetworkInferenceQA/Private/ModelUnitTester.h
+++ b/Engine/Plugins/Experimental/NeuralNetworkInference/Source/NeuralNetworkInferenceQA/Private/ModelUnitTester.h
@@ -26,7 +26,7 @@ private:
 	static FString GetUAssetModelFilePath(const FString& InModelName, const FString& InModelZooRelativeDirectory);
 	static UNeuralNetwork* NetworkUassetLoadTest(const FString& InUAssetPath);
 	static UNeuralNetwork* NetworkONNXOrORTLoadTest(const FString& InModelFilePath);
-	static bool ModelAccuracyTest(UNeuralNetwork* InOutNetwork, const ENeuralNetworkSynchronousMode InSynchronousMode, const UNeuralNetwork::ENeuralBackEnd InBackEnd, const TArray<float>& InInputArrayValues,
+	static bool ModelAccuracyTest(UNeuralNetwork* InOutNetwork, const ENeuralSynchronousMode InSynchronousMode, const UNeuralNetwork::ENeuralBackEnd InBackEnd, const TArray<float>& InInputArrayValues,
 		const TArray<double>& InCPUGroundTruths, const TArray<double>& InGPUGroundTruths);
 	static void ModelAccuracyTestRun(TArray<TArray<float>>& OutOutputs, UNeuralNetwork* InOutNetwork, const TArray<TArray<float>> InInputArrays,
 		const ENeuralDeviceType InDeviceType, const ENeuralDeviceType InInputDeviceType, const ENeuralDeviceType InOutputDeviceType);
diff --git a/Engine/Plugins/Experimental/NeuralNetworkInference/Source/NeuralNetworkInferenceQA/Private/OperatorTester.cpp b/Engine/Plugins/Experimental/NeuralNetworkInference/Source/NeuralNetworkInferenceQA/Private/OperatorTester.cpp
--- a/Engine/Plugins/Experimental/NeuralNetworkInference/Source/NeuralNetworkInferenceQA/Private/OperatorTester.cpp
+++ b/Engine/Plugins/Experimental/NeuralNetworkInference/Source/NeuralNetworkInferenceQA/Private/OperatorTester.cpp
@@ -63,7 +63,7 @@ void FOperatorTester::TestOperator(UNeuralNetworkInferenceQAAsset* InOutNetworkI
 			for (int64 AuxiliaryTensorIndex = 0; AuxiliaryTensorIndex < Operator->GetNumberAuxiliaryTensors(); ++AuxiliaryTensorIndex)
 			{
 				const FString AuxiliaryTensorName = Operator->GetName() + TEXT("-Aux") + FString::FromInt(AuxiliaryTensorIndex);
-				InOutTensors[Index].Push(FNeuralTensor(ENeuralDataType::None, TArray<int64>({}), AuxiliaryTensorName, ENeuralTensorTypeGPU::Generic)); // ENeuralTensorTypeGPU can be modify by FNeuralOperator::Configure()
+				InOutTensors[Index].Push(FNeuralTensor(ENeuralDataType::None, TArray<int64>({}), AuxiliaryTensorName, ENeuralTensorType::Generic)); // ENeuralTensorType can be modify by FNeuralOperator::Configure()
 				AuxiliaryOperatorTensors.Push(&InOutTensors[Index].Last());
 			}
 			Operator->SetAuxiliaryTensors(AuxiliaryOperatorTensors);
