Adding NNI documentation - Part 3:
- NNI/Public/NeuraTensor.h cleaned a bit

#rnx

diff --git a/Engine/Plugins/Experimental/NeuralNetworkInference/Source/NeuralNetworkInference/Private/NeuralTensor.cpp b/Engine/Plugins/Experimental/NeuralNetworkInference/Source/NeuralNetworkInference/Private/NeuralTensor.cpp
--- a/Engine/Plugins/Experimental/NeuralNetworkInference/Source/NeuralNetworkInference/Private/NeuralTensor.cpp
+++ b/Engine/Plugins/Experimental/NeuralNetworkInference/Source/NeuralNetworkInference/Private/NeuralTensor.cpp
@@ -7,11 +7,8 @@
 #include "NeuralNetworkInferenceUtils.h"
 #include "NeuralNetworkInferenceUtilsGPU.h"
 #include "NeuralTensorResourceArray.h"
-#include "RenderGraphBuilder.h"
-#include "RenderGraphResources.h"
 #include "RenderGraphUtils.h"
 #include "RenderingThread.h"
-#include "Templates/RefCounting.h"
 
 
 
diff --git a/Engine/Plugins/Experimental/NeuralNetworkInference/Source/NeuralNetworkInference/Public/NeuralEnumClasses.h b/Engine/Plugins/Experimental/NeuralNetworkInference/Source/NeuralNetworkInference/Public/NeuralEnumClasses.h
--- a/Engine/Plugins/Experimental/NeuralNetworkInference/Source/NeuralNetworkInference/Public/NeuralEnumClasses.h
+++ b/Engine/Plugins/Experimental/NeuralNetworkInference/Source/NeuralNetworkInference/Public/NeuralEnumClasses.h
@@ -8,7 +8,7 @@
 #include "NeuralEnumClasses.generated.h"
 
 /**
- * It defines the data type (float, double, int32, etc.).
+ * It defines the data type (e.g., float, double, int32).
  */
 UENUM()
 enum class ENeuralDataType : uint8
diff --git a/Engine/Plugins/Experimental/NeuralNetworkInference/Source/NeuralNetworkInference/Public/NeuralTensor.h b/Engine/Plugins/Experimental/NeuralNetworkInference/Source/NeuralNetworkInference/Public/NeuralTensor.h
--- a/Engine/Plugins/Experimental/NeuralNetworkInference/Source/NeuralNetworkInference/Public/NeuralTensor.h
+++ b/Engine/Plugins/Experimental/NeuralNetworkInference/Source/NeuralNetworkInference/Public/NeuralTensor.h
@@ -4,21 +4,19 @@
 
 #include "CoreMinimal.h"
 #include "NeuralEnumClasses.h"
+#include "RenderGraphBuilder.h" // FRDGBuilder
+#include "RenderGraphDefinitions.h" // FRDGBufferSRVRef, FRDGBufferUAVRef
+#include "RenderGraphResources.h" // FRDGPooledBuffer
+#include "RHIDefinitions.h" // EBufferUsageFlags
+#include "Templates/RefCounting.h" // TRefCountPtr
 #include <algorithm>
 #include "NeuralTensor.generated.h"
 
-class FRDGBuilder;
-using FRDGBufferSRVRef = class FRDGBufferSRV*;
-using FRDGBufferUAVRef = class FRDGBufferUAV*;
-template<typename T> class TRefCountPtr;
-enum class EBufferUsageFlags : uint32;
-struct FTensorProto;
-
 /**
- * This is an auxiliary class. See UNeuralNetworkLegacy for a high-level wrapper of the whole NeuralNetworkInference plugin. The UNeuralNetworkLegacy header
+ * This is an auxiliary class. See UNeuralNetwork for a high-level wrapper of the whole NeuralNetworkInference plugin. The UNeuralNetwork header
  * documentation also includes some code examples.
  *
- * FNeuralTensor represents a tensor of the UNeuralNetworkLegacy model.
+ * FNeuralTensor represents a tensor of the UNeuralNetwork model.
  * Its CPU operations are very similar to those of TArray<T>. Most of its operations run on CPU, so `ToGPU_RenderThread()` must be called before running on
  * GPU and after running any FNeuralTensor function that modified the CPU memory.
  */
@@ -184,7 +182,7 @@ public:
 	 */
 	bool InitPooledBuffer(void** NativeResource = nullptr);
 
-	TRefCountPtr<class FRDGPooledBuffer>& GetPooledBuffer() const;
+	TRefCountPtr<FRDGPooledBuffer>& GetPooledBuffer() const;
 	const FRDGBufferSRVRef GetBufferSRVRef() const;
 	FRDGBufferUAVRef GetBufferUAVRef();
 
@@ -216,7 +214,7 @@ public:
 	 * It will fill the current neural tensor with the input FTensorProto.
 	 * @returns Whether the conversion was successful.
 	 */
-	bool SetFromTensorProto(const FTensorProto* const InTensorProto, const FString& InTensorName, const ENeuralTensorType InTensorTypeGPU);
+	bool SetFromTensorProto(const struct FTensorProto* const InTensorProto, const FString& InTensorName, const ENeuralTensorType InTensorTypeGPU);
 
 	/**
 	 * It sets all the elements of the FNeuralTensor to InValue.
@@ -281,7 +279,7 @@ private:
 	 * - ArrayCPUForGPUAs32Data: If ArrayCPU is meant for 64-byte data (i.e., int64, uint64, double).
 	 */
 	bool bEnableGPU;
-	TSharedPtr<TRefCountPtr<class FRDGPooledBuffer>> PooledBuffer;
+	TSharedPtr<TRefCountPtr<FRDGPooledBuffer>> PooledBuffer;
 	TSharedPtr<FRDGBufferSRVRef> BufferSRVRef;
 	TSharedPtr<FRDGBufferUAVRef> BufferUAVRef;
 	TArray<uint8> ArrayCPUForGPUAs32Data;
diff --git a/Engine/Plugins/Experimental/NeuralNetworkInference/Source/NeuralNetworkInferenceQA/Internal/NeuralNetworkInferenceQAAsset.h b/Engine/Plugins/Experimental/NeuralNetworkInference/Source/NeuralNetworkInferenceQA/Internal/NeuralNetworkInferenceQAAsset.h
--- a/Engine/Plugins/Experimental/NeuralNetworkInference/Source/NeuralNetworkInferenceQA/Internal/NeuralNetworkInferenceQAAsset.h
+++ b/Engine/Plugins/Experimental/NeuralNetworkInference/Source/NeuralNetworkInferenceQA/Internal/NeuralNetworkInferenceQAAsset.h
@@ -30,7 +30,8 @@ public:
 	/**
 	 * It will compare new vs. previous tests.
 	 */
-	static bool CompareAverageL1DiffNewVsPreviousTests(const FNeuralNetworkInferenceQAOperatorTestAsset& InNeuralNetworkInferenceQAOperatorTestAsset1, const FNeuralNetworkInferenceQAOperatorTestAsset& InNeuralNetworkInferenceQAOperatorTestAsset2, const FString& InOperatorName);
+	static bool CompareAverageL1DiffNewVsPreviousTests(const FNeuralNetworkInferenceQAOperatorTestAsset& InOperatorTestAsset1,
+		const FNeuralNetworkInferenceQAOperatorTestAsset& InOperatorTestAsset2, const FString& InOperatorName);
 
 private:
 	/**
diff --git a/Engine/Plugins/Experimental/NeuralNetworkInference/Source/NeuralNetworkInferenceQA/Private/NeuralNetworkInferenceQAAsset.cpp b/Engine/Plugins/Experimental/NeuralNetworkInference/Source/NeuralNetworkInferenceQA/Private/NeuralNetworkInferenceQAAsset.cpp
--- a/Engine/Plugins/Experimental/NeuralNetworkInference/Source/NeuralNetworkInferenceQA/Private/NeuralNetworkInferenceQAAsset.cpp
+++ b/Engine/Plugins/Experimental/NeuralNetworkInference/Source/NeuralNetworkInferenceQA/Private/NeuralNetworkInferenceQAAsset.cpp
@@ -27,11 +27,12 @@ void FNeuralNetworkInferenceQAOperatorTestAsset::AddOutputTensors(const UNeuralN
 	}
 }
 
-bool FNeuralNetworkInferenceQAOperatorTestAsset::CompareAverageL1DiffNewVsPreviousTests(const FNeuralNetworkInferenceQAOperatorTestAsset& InNeuralNetworkInferenceQAOperatorTestAsset1, const FNeuralNetworkInferenceQAOperatorTestAsset& InNeuralNetworkInferenceQAOperatorTestAsset2, const FString& InOperatorName)
+bool FNeuralNetworkInferenceQAOperatorTestAsset::CompareAverageL1DiffNewVsPreviousTests(const FNeuralNetworkInferenceQAOperatorTestAsset& InOperatorTestAsset1,
+	const FNeuralNetworkInferenceQAOperatorTestAsset& InOperatorTestAsset2, const FString& InOperatorName)
 {
 	// Compare each input and output
-	return (CompareAverageL1DiffNewVsPreviousTests(InNeuralNetworkInferenceQAOperatorTestAsset1.InputTensors, InNeuralNetworkInferenceQAOperatorTestAsset2.InputTensors, InOperatorName)
-		&& CompareAverageL1DiffNewVsPreviousTests(InNeuralNetworkInferenceQAOperatorTestAsset1.OutputTensors, InNeuralNetworkInferenceQAOperatorTestAsset2.OutputTensors, InOperatorName));
+	return (CompareAverageL1DiffNewVsPreviousTests(InOperatorTestAsset1.InputTensors, InOperatorTestAsset2.InputTensors, InOperatorName)
+		&& CompareAverageL1DiffNewVsPreviousTests(InOperatorTestAsset1.OutputTensors, InOperatorTestAsset2.OutputTensors, InOperatorName));
 }
 
 bool FNeuralNetworkInferenceQAOperatorTestAsset::CompareAverageL1DiffNewVsPreviousTests(const TArray<FNeuralTensor>& InTensorsA, const TArray<FNeuralTensor>& InTensorsB, const FString& InOperatorName)
