Updated deprecated UE define by its new enum class equivalent (BUF_ replaced by EBufferUsageFlags::)
NNI/Public/NeuraTensor.h reduced to up to 150-characters per line

#rnx

diff --git a/Engine/Plugins/Experimental/NeuralNetworkInference/Source/NeuralNetworkInference/Private/NeuralTensor.cpp b/Engine/Plugins/Experimental/NeuralNetworkInference/Source/NeuralNetworkInference/Private/NeuralTensor.cpp
--- a/Engine/Plugins/Experimental/NeuralNetworkInference/Source/NeuralNetworkInference/Private/NeuralTensor.cpp
+++ b/Engine/Plugins/Experimental/NeuralNetworkInference/Source/NeuralNetworkInference/Private/NeuralTensor.cpp
@@ -269,45 +269,45 @@ void FNeuralTensor::SetTensorTypeGPU(const ENeuralTensorType InTensorTypeGPU)
 
 void FNeuralTensor::ToGPU_RenderThread(FRDGBuilder* InOutGraphBuilder)
 {
-// @todo: Volatile or not adding BUF_ShaderResource causes this error:
+// @todo: Volatile or not adding EBufferUsageFlags::ShaderResource causes this error:
 //Fatal error: [File:D:/P4/private_dh_research_pitt/Engine/Source/Runtime/Windows/D3D11RHI/Private/D3D11Util.cpp] [Line: 258] 
 //Result failed 
 // at D:/P4/private_dh_research_pitt/Engine/Source/Runtime/Windows/D3D11RHI/Private/D3D11VertexBuffer.cpp:109 
 // with error E_INVALIDARG
 	// Idea:
-	// - BUF_Volatile: Updated multiple times in a frame, but does not imply a lifetime of 1 frame. E.g. a vertex buffer you update every frame with new vertices.
-	// - BUF_Transient: Used during 1 frame. Volatile and transient are not mutually exclusive.
-	// - BUF_KeepCPUAccessible: Not needed, I can just copy the final GPU memory back to RAM at the very end
+	// - EBufferUsageFlags::Volatile: Updated multiple times in a frame, but does not imply a lifetime of 1 frame. E.g. a vertex buffer you update every frame with new vertices.
+	// - EBufferUsageFlags::Transient: Used during 1 frame. Volatile and transient are not mutually exclusive.
+	// - EBufferUsageFlags::KeepCPUAccessible: Not needed, I can just copy the final GPU memory back to RAM at the very end
 	// - Input and Intermediate(Not)Initialized currently share the same attributes because input might become intermediate (e.g., if input tensor fed into a ReLU, which simply modifies
 	//   the input FNeuralTensor). However, Intermediate(Not)Initialized and Output do not copy the memory from CPU to GPU but rather simply allocates it.
 	// - Output might also become Intermediate(Not)Initialized (e.g., if Output -> ReLU -> Output), so it is kept as ReadWrite rather than written once to account for this.
 	// Call ToGPU_RenderThread with the right flags
 	if (TensorTypeGPU == ENeuralTensorType::Generic)
 	{
-		return ToGPU_RenderThread(InOutGraphBuilder, BUF_ShaderResource | BUF_UnorderedAccess, true);
+		return ToGPU_RenderThread(InOutGraphBuilder, EBufferUsageFlags::ShaderResource | EBufferUsageFlags::UnorderedAccess, true);
 	}
 	else if (TensorTypeGPU == ENeuralTensorType::Input)
 	{
-		return ToGPU_RenderThread(InOutGraphBuilder, BUF_ShaderResource | BUF_UnorderedAccess, true);
+		return ToGPU_RenderThread(InOutGraphBuilder, EBufferUsageFlags::ShaderResource | EBufferUsageFlags::UnorderedAccess, true);
 	}
 	else if (TensorTypeGPU == ENeuralTensorType::IntermediateNotInitialized)
 	{
-		return ToGPU_RenderThread(InOutGraphBuilder, BUF_ShaderResource | BUF_UnorderedAccess | BUF_Transient, false);
+		return ToGPU_RenderThread(InOutGraphBuilder, EBufferUsageFlags::ShaderResource | EBufferUsageFlags::UnorderedAccess | EBufferUsageFlags::Transient, false);
 	}
 	else if (TensorTypeGPU == ENeuralTensorType::IntermediateInitialized)
 	{
-		return ToGPU_RenderThread(InOutGraphBuilder, BUF_ShaderResource | BUF_UnorderedAccess, true);
+		return ToGPU_RenderThread(InOutGraphBuilder, EBufferUsageFlags::ShaderResource | EBufferUsageFlags::UnorderedAccess, true);
 	}
 	else if (TensorTypeGPU == ENeuralTensorType::Output)
 	{
-		return ToGPU_RenderThread(InOutGraphBuilder, BUF_ShaderResource | BUF_UnorderedAccess, false);
+		return ToGPU_RenderThread(InOutGraphBuilder, EBufferUsageFlags::ShaderResource | EBufferUsageFlags::UnorderedAccess, false);
 	}
 	else if (TensorTypeGPU == ENeuralTensorType::Weight)
 	{
-		return ToGPU_RenderThread(InOutGraphBuilder, BUF_ShaderResource | BUF_Static, true);
+		return ToGPU_RenderThread(InOutGraphBuilder, EBufferUsageFlags::ShaderResource | EBufferUsageFlags::Static, true);
 	}
 	UE_LOG(LogNeuralNetworkInference, Warning, TEXT("FNeuralTensor-%s::ToGPU_RenderThread(): Unimplemented TensorTypeGPU = %d. Assuming ENeuralTensorType::Generic."), *Name, (int32)TensorTypeGPU);
-	return ToGPU_RenderThread(InOutGraphBuilder, BUF_UnorderedAccess, true);
+	return ToGPU_RenderThread(InOutGraphBuilder, EBufferUsageFlags::UnorderedAccess, true);
 }
 
 void FNeuralTensor::ToGPU_RenderThread(FRDGBuilder* InOutGraphBuilder, const EBufferUsageFlags InEBufferUsageFlags, const bool bInShouldCopyFromCPU)
@@ -321,10 +321,10 @@ void FNeuralTensor::ToGPU_RenderThread(FRDGBuilder* InOutGraphBuilder, const EBu
 	checkf(IsInRenderingThread(), TEXT("FNeuralTensor-%s::ToGPU_RenderThread(): IsInRenderingThread() must be true."), *Name);
 	checkf(InOutGraphBuilder, TEXT("FNeuralTensor-%s::ToGPU_RenderThread(): InOutGraphBuilder is nullptr."), *Name);
 	// Not SRV-only and not UAV/SRV
-	checkf(EnumHasAnyFlags(InEBufferUsageFlags, BUF_ShaderResource|BUF_UnorderedAccess),
+	checkf(EnumHasAnyFlags(InEBufferUsageFlags, EBufferUsageFlags::ShaderResource|EBufferUsageFlags::UnorderedAccess),
 		TEXT("FNeuralTensor-%s::ToGPU_RenderThread(): Unexpected case InEBufferUsageFlags = %d."), *Name, (uint32)InEBufferUsageFlags);
 	// If SRV-only and bInShouldCopyFromCPU == false
-	if (!EnumHasAnyFlags(InEBufferUsageFlags, BUF_UnorderedAccess) && !bInShouldCopyFromCPU)
+	if (!EnumHasAnyFlags(InEBufferUsageFlags, EBufferUsageFlags::UnorderedAccess) && !bInShouldCopyFromCPU)
 	{
 		UE_LOG(LogNeuralNetworkInference, Warning, TEXT("FNeuralTensor-%s::ToGPU_RenderThread(): bInShouldCopyFromCPU must be true for SRVs (because they cannot be edited). Assumed true."), *Name);
 	}
@@ -356,7 +356,7 @@ void FNeuralTensor::ToGPU_RenderThread(FRDGBuilder* InOutGraphBuilder, const EBu
 	}
 
 	// Recreate BufferSRVRef
-	if (EnumHasAnyFlags(InEBufferUsageFlags, BUF_ShaderResource))
+	if (EnumHasAnyFlags(InEBufferUsageFlags, EBufferUsageFlags::ShaderResource))
 	{
 		BufferSRVRef = MakeShared<FRDGBufferSRVRef>(InOutGraphBuilder->CreateSRV(BufferRef, FNeuralDataTypeUtils::GetPixelFormat(DataType)));
 	}
@@ -365,7 +365,7 @@ void FNeuralTensor::ToGPU_RenderThread(FRDGBuilder* InOutGraphBuilder, const EBu
 		BufferSRVRef.Reset();
 	}
 	// Recreate BufferUAVRef
-	if (EnumHasAnyFlags(InEBufferUsageFlags, BUF_UnorderedAccess))
+	if (EnumHasAnyFlags(InEBufferUsageFlags, EBufferUsageFlags::UnorderedAccess))
 	{
 		BufferUAVRef = MakeShared<FRDGBufferUAVRef>(InOutGraphBuilder->CreateUAV(BufferRef, FNeuralDataTypeUtils::GetPixelFormat(DataType)));
 	}
@@ -443,7 +443,7 @@ bool FNeuralTensor::InitPooledBuffer(void** NativeResource)
 			FRDGBufferDesc BufferDesc;
 			BufferDesc.BytesPerElement = FNeuralDataTypeUtils::GetByteSize(DataType);
 			BufferDesc.NumElements = Num();
-			BufferDesc.Usage = BUF_UnorderedAccess | BUF_ShaderResource;
+			BufferDesc.Usage = EBufferUsageFlags::UnorderedAccess | EBufferUsageFlags::ShaderResource;
 			BufferDesc.UnderlyingType = FRDGBufferDesc::EUnderlyingType::VertexBuffer;
 
 			FRDGBufferRef BufferRef = Builder.CreateBuffer(BufferDesc, GetNameData());
diff --git a/Engine/Plugins/Experimental/NeuralNetworkInference/Source/NeuralNetworkInference/Public/NeuralTensor.h b/Engine/Plugins/Experimental/NeuralNetworkInference/Source/NeuralNetworkInference/Public/NeuralTensor.h
--- a/Engine/Plugins/Experimental/NeuralNetworkInference/Source/NeuralNetworkInference/Public/NeuralTensor.h
+++ b/Engine/Plugins/Experimental/NeuralNetworkInference/Source/NeuralNetworkInference/Public/NeuralTensor.h
@@ -9,7 +9,7 @@
 #include "RenderGraphResources.h" // FRDGPooledBuffer
 #include "RHIDefinitions.h" // EBufferUsageFlags
 #include "Templates/RefCounting.h" // TRefCountPtr
-#include <algorithm>
+#include <algorithm> // std::fill
 #include "NeuralTensor.generated.h"
 
 /**
@@ -17,8 +17,8 @@
  * documentation also includes some code examples.
  *
  * FNeuralTensor represents a tensor of the UNeuralNetwork model.
- * Its CPU operations are very similar to those of TArray<T>. Most of its operations run on CPU, so `ToGPU_RenderThread()` must be called before running on
- * GPU and after running any FNeuralTensor function that modified the CPU memory.
+ * Its CPU operations are very similar to those of TArray<T>. Most of its operations run on CPU, so `ToGPU_RenderThread()` must be called before
+ * running on GPU and after running any FNeuralTensor function that modified the CPU memory.
  */
 USTRUCT()
 struct NEURALNETWORKINFERENCE_API FNeuralTensor
@@ -32,8 +32,10 @@ public:
 	 * @param InSizes Set to empty (or omit argument) if memory allocation is not required or the final size is unknown.
 	 * @param InName Used for GPU debugging and the ToString() function.
 	*/
-	explicit FNeuralTensor(const ENeuralDataType InDataType = ENeuralDataType::None, const TArray<int64>& InSizes = TArray<int64>(), const FString& InName = TEXT("FNeuralTensor"), const ENeuralTensorType InTensorTypeGPU = ENeuralTensorType::Generic);
-	explicit FNeuralTensor(const ENeuralDataType InDataType, const int64 InVolume, const FString& InName = TEXT("FNeuralTensor"), const ENeuralTensorType InTensorTypeGPU = ENeuralTensorType::Generic);
+	explicit FNeuralTensor(const ENeuralDataType InDataType = ENeuralDataType::None, const TArray<int64>& InSizes = TArray<int64>(),
+		const FString& InName = TEXT("FNeuralTensor"), const ENeuralTensorType InTensorTypeGPU = ENeuralTensorType::Generic);
+	explicit FNeuralTensor(const ENeuralDataType InDataType, const int64 InVolume, const FString& InName = TEXT("FNeuralTensor"),
+		const ENeuralTensorType InTensorTypeGPU = ENeuralTensorType::Generic);
 	explicit FNeuralTensor(const FString& InName, const ENeuralTensorType InTensorTypeGPU = ENeuralTensorType::Generic);
 
 	/**
@@ -45,11 +47,12 @@ public:
 	 * @param InName Used for GPU debugging and the ToString() function.
 	 */
 	template <typename T>
-	explicit FNeuralTensor(const TArray<T>& InArray, const TArray<int64>& InSizes = TArray<int64>(), const FString& InName = TEXT("FNeuralTensor"), const ENeuralTensorType InTensorTypeGPU = ENeuralTensorType::Generic);
+	explicit FNeuralTensor(const TArray<T>& InArray, const TArray<int64>& InSizes = TArray<int64>(), const FString& InName = TEXT("FNeuralTensor"),
+		const ENeuralTensorType InTensorTypeGPU = ENeuralTensorType::Generic);
 
 	/**
 	 * Copy constructor.
-	 * It performs a deep (slow but safe) copy of the current FNeuralTensor. The resulting FNeuralTensor will not share any parameters with the current one.
+	 * It performs a deep (slow but safe) copy of the current tensor. The resulting tensor will not share any parameters with the current one.
 	 */
 	FNeuralTensor(const FNeuralTensor& InTensor);
 
@@ -97,8 +100,8 @@ public:
 	/**
 	 * 2 different functions to get the output tensor results:
 	 * - GetArrayCopy<T> (slower but safer): It returns a copy of the results as a TArray<T>. T has to be the same size than sizeof(DataType).
-	 * - GetUnderlyingUInt8ArrayRef (faster): It returns a reference to the TArray<uint8> that contains the results. The returned reference will only be valid as long as
-	 *   this class instance does not go out of scope and its internal TArray<uint8> is not re-initialized.
+	 * - GetUnderlyingUInt8ArrayRef (faster): It returns a reference to the TArray<uint8> that contains the results. The returned reference will
+	 *   only be valid as long as this class instance does not go out of scope and its internal TArray<uint8> is not re-initialized.
 	 * @see GetArrayCopy(), GetUnderlyingUInt8ArrayRef()
 	 */
 	template <typename T>
@@ -169,7 +172,8 @@ public:
 
 	/**
 	 * It uploads/downloads the memory from/to the CPU to/from the GPU based on TensorTypeGPU (which sets a preset of EBufferUsageFlags flags).
-	 * @param InEBufferUsageFlags gives the user total control over the buffer flags (and ignores the TensorTypeGPU flag). This is meant to be filled with a combination of EBufferUsageFlags values.
+	 * @param InEBufferUsageFlags gives the user total control over the buffer flags (and ignores the TensorTypeGPU flag). This is meant to be filled
+	 * with a combination of EBufferUsageFlags values.
 	 */
 	void ToGPU_RenderThread(FRDGBuilder* InOutGraphBuilder);
 	void ToGPU_RenderThread(FRDGBuilder* InOutGraphBuilder, const EBufferUsageFlags InEBufferUsageFlags, const bool bInShouldCopyFromCPU);
@@ -190,7 +194,8 @@ public:
 	 * Resize the FNeuralTensor to the desired new size.
 	 * Analog to TArray<T>::SetNumUninitialized().
 	 * @param InTensor Set Sizes and DataType from the input InTensor.
-	 * @param InVolume Set InVolume to 0 if memory allocation is not required or the final size is unknown, SetNumUninitialized(InVolume) == SetNumUninitialized({InVolume}).
+	 * @param InVolume Set InVolume to 0 if memory allocation is not required or the final size is unknown,
+	 * SetNumUninitialized(InVolume) == SetNumUninitialized({InVolume}).
 	 * @param InSizes Set InSizes to empty if memory allocation is not required or the final size is unknown.
 	 * @param InDataType set to None means that it will maintain the previous type.
 	 */
@@ -250,7 +255,8 @@ public:
 
 	/**
 	 * It returns a FString with up to InMaxNumberElementsToDisplay elements displayed. If InMaxNumberElementsToDisplay <= 0, it displays them all.
-	 * @param bInReturnOnlyData If false (default), it will print all information. E.g., "FNeuralTensor: Int64, Generic, volume=3, sizes={3}, data=[1 2 3]".
+	 * @param bInReturnOnlyData If false (default), it will print all information.
+	 * E.g., "FNeuralTensor: Int64, Generic, volume=3, sizes={3}, data=[1 2 3]".
 	 * If false, it will simply print the data. I.e., "[1 2 3]" for the previous example.
 	 */
 	FString ToString(const int64 InMaxNumberElementsToDisplay = -1, const bool bInReturnOnlyData = false) const;
@@ -275,7 +281,8 @@ private:
 	TArray<uint8> ArrayCPU;
 	/**
 	 * GPU-based variables that are transient (do not need to be serialized) and initialized on the fly (with ToGPU/ToCPU).
-	 * - bEnableGPU: By default false, meaning all GPU memory will be disabled and those functions will not do anything. Enable to allow using the GPU functions and variables of the tensor.
+	 * - bEnableGPU: By default false, meaning all GPU memory will be disabled and those functions will not do anything. Enable to allow using the
+	 *   GPU functions and variables of the tensor.
 	 * - ArrayCPUForGPUAs32Data: If ArrayCPU is meant for 64-byte data (i.e., int64, uint64, double).
 	 */
 	bool bEnableGPU;
@@ -288,7 +295,8 @@ private:
 	 * Auxiliary for the templated constructor FNeuralTensor(const ENeuralDataType InDataType, const TArray<T>& InArray, ...).
 	 * InSizeOfT and sizeof(InDataType) should match, as well as the volume of InValueNum and InSizes.
 	 */
-	explicit FNeuralTensor(const ENeuralDataType InDataType, const void* const InValues, const int64 InSizeOfT, const int64 InValueNum, const TArray<int64>& InSizes, const FString& InName, const ENeuralTensorType InTensorTypeGPU);
+	explicit FNeuralTensor(const ENeuralDataType InDataType, const void* const InValues, const int64 InSizeOfT, const int64 InValueNum,
+		const TArray<int64>& InSizes, const FString& InName, const ENeuralTensorType InTensorTypeGPU);
 
 	/**
 	 * Auxiliary for the templated SetFromArrayCopy().
@@ -305,12 +313,13 @@ private:
 
 
 
-/* FNeuralTensor inline functions
+/* FNeuralTensor inlined and templated functions
  *****************************************************************************/
 
 template <typename T>
 FNeuralTensor::FNeuralTensor(const TArray<T>& InArray, const TArray<int64>& InSizes, const FString& InName, const ENeuralTensorType InTensorTypeGPU)
-	: FNeuralTensor(FNeuralDataTypeUtils::GetDataType<T>(), InArray.GetData(), sizeof(T), InArray.Num(), (InSizes.Num() > 0 ? InSizes : TArray<int64>({ (int64)InArray.Num() })), InName, InTensorTypeGPU)
+	: FNeuralTensor(FNeuralDataTypeUtils::GetDataType<T>(), InArray.GetData(), sizeof(T), InArray.Num(),
+		(InSizes.Num() > 0 ? InSizes : TArray<int64>({ (int64)InArray.Num() })), InName, InTensorTypeGPU)
 {}
 
 bool FNeuralTensor::operator!=(const FNeuralTensor& InTensorToCopy) const
