Fixed typo in doc.

_codereview Francisco.Vicente
#rnx

diff --git a/Engine/Plugins/Experimental/NeuralNetworkInference/Source/NeuralNetworkInference/Internal/NeuralInt64ArrayUInt32Buffer.h b/Engine/Plugins/Experimental/NeuralNetworkInference/Source/NeuralNetworkInference/Internal/NeuralInt64ArrayUInt32Buffer.h
--- a/Engine/Plugins/Experimental/NeuralNetworkInference/Source/NeuralNetworkInference/Internal/NeuralInt64ArrayUInt32Buffer.h
+++ b/Engine/Plugins/Experimental/NeuralNetworkInference/Source/NeuralNetworkInference/Internal/NeuralInt64ArrayUInt32Buffer.h
@@ -79,7 +79,7 @@ private:
 
 
 
-/* FNeuralInt64ArrayUInt32Buffer inline functions
+/* FNeuralInt64ArrayUInt32Buffer inlined and templated functions
  *****************************************************************************/
 
 int32 FNeuralInt64ArrayUInt32Buffer::Num() const
diff --git a/Engine/Plugins/Experimental/NeuralNetworkInference/Source/NeuralNetworkInference/Internal/NeuralOperator.h b/Engine/Plugins/Experimental/NeuralNetworkInference/Source/NeuralNetworkInference/Internal/NeuralOperator.h
--- a/Engine/Plugins/Experimental/NeuralNetworkInference/Source/NeuralNetworkInference/Internal/NeuralOperator.h
+++ b/Engine/Plugins/Experimental/NeuralNetworkInference/Source/NeuralNetworkInference/Internal/NeuralOperator.h
@@ -153,7 +153,7 @@ private:
 
 
 
-/* FNeuralOperator inline functions
+/* FNeuralOperator inlined and templated functions
  *****************************************************************************/
 
 FString FNeuralOperator::GetName() const
diff --git a/Engine/Plugins/Experimental/NeuralNetworkInference/Source/NeuralNetworkInference/Internal/NeuralOperators/AbsOperator.h b/Engine/Plugins/Experimental/NeuralNetworkInference/Source/NeuralNetworkInference/Internal/NeuralOperators/AbsOperator.h
--- a/Engine/Plugins/Experimental/NeuralNetworkInference/Source/NeuralNetworkInference/Internal/NeuralOperators/AbsOperator.h
+++ b/Engine/Plugins/Experimental/NeuralNetworkInference/Source/NeuralNetworkInference/Internal/NeuralOperators/AbsOperator.h
@@ -17,7 +17,7 @@ public:
 
 
 
-/* FAbsOperator inline functions
+/* FAbsOperator inlined and templated functions
  *****************************************************************************/
 
 void FAbsOperator::ForwardCPU()
diff --git a/Engine/Plugins/Experimental/NeuralNetworkInference/Source/NeuralNetworkInference/Internal/NeuralOperators/AcosOperator.h b/Engine/Plugins/Experimental/NeuralNetworkInference/Source/NeuralNetworkInference/Internal/NeuralOperators/AcosOperator.h
--- a/Engine/Plugins/Experimental/NeuralNetworkInference/Source/NeuralNetworkInference/Internal/NeuralOperators/AcosOperator.h
+++ b/Engine/Plugins/Experimental/NeuralNetworkInference/Source/NeuralNetworkInference/Internal/NeuralOperators/AcosOperator.h
@@ -17,7 +17,7 @@ public:
 
 
 
-/* FAcosOperator inline functions
+/* FAcosOperator inlined and templated functions
  *****************************************************************************/
 
 void FAcosOperator::ForwardCPU()
diff --git a/Engine/Plugins/Experimental/NeuralNetworkInference/Source/NeuralNetworkInference/Internal/NeuralOperators/AddOperator.h b/Engine/Plugins/Experimental/NeuralNetworkInference/Source/NeuralNetworkInference/Internal/NeuralOperators/AddOperator.h
--- a/Engine/Plugins/Experimental/NeuralNetworkInference/Source/NeuralNetworkInference/Internal/NeuralOperators/AddOperator.h
+++ b/Engine/Plugins/Experimental/NeuralNetworkInference/Source/NeuralNetworkInference/Internal/NeuralOperators/AddOperator.h
@@ -17,7 +17,7 @@ public:
 
 
 
-/* FAddOperator inline functions
+/* FAddOperator inlined and templated functions
  *****************************************************************************/
 
 void FAddOperator::ForwardCPU()
diff --git a/Engine/Plugins/Experimental/NeuralNetworkInference/Source/NeuralNetworkInference/Internal/NeuralOperators/AsinOperator.h b/Engine/Plugins/Experimental/NeuralNetworkInference/Source/NeuralNetworkInference/Internal/NeuralOperators/AsinOperator.h
--- a/Engine/Plugins/Experimental/NeuralNetworkInference/Source/NeuralNetworkInference/Internal/NeuralOperators/AsinOperator.h
+++ b/Engine/Plugins/Experimental/NeuralNetworkInference/Source/NeuralNetworkInference/Internal/NeuralOperators/AsinOperator.h
@@ -17,7 +17,7 @@ public:
 
 
 
-/* FAsinOperator inline functions
+/* FAsinOperator inlined and templated functions
  *****************************************************************************/
 
 void FAsinOperator::ForwardCPU()
diff --git a/Engine/Plugins/Experimental/NeuralNetworkInference/Source/NeuralNetworkInference/Internal/NeuralOperators/AtanOperator.h b/Engine/Plugins/Experimental/NeuralNetworkInference/Source/NeuralNetworkInference/Internal/NeuralOperators/AtanOperator.h
--- a/Engine/Plugins/Experimental/NeuralNetworkInference/Source/NeuralNetworkInference/Internal/NeuralOperators/AtanOperator.h
+++ b/Engine/Plugins/Experimental/NeuralNetworkInference/Source/NeuralNetworkInference/Internal/NeuralOperators/AtanOperator.h
@@ -17,7 +17,7 @@ public:
 
 
 
-/* FAtanOperator inline functions
+/* FAtanOperator inlined and templated functions
  *****************************************************************************/
 
 void FAtanOperator::ForwardCPU()
diff --git a/Engine/Plugins/Experimental/NeuralNetworkInference/Source/NeuralNetworkInference/Internal/NeuralOperators/CeilOperator.h b/Engine/Plugins/Experimental/NeuralNetworkInference/Source/NeuralNetworkInference/Internal/NeuralOperators/CeilOperator.h
--- a/Engine/Plugins/Experimental/NeuralNetworkInference/Source/NeuralNetworkInference/Internal/NeuralOperators/CeilOperator.h
+++ b/Engine/Plugins/Experimental/NeuralNetworkInference/Source/NeuralNetworkInference/Internal/NeuralOperators/CeilOperator.h
@@ -17,7 +17,7 @@ public:
 
 
 
-/* FCeilOperator inline functions
+/* FCeilOperator inlined and templated functions
  *****************************************************************************/
 
 void FCeilOperator::ForwardCPU()
diff --git a/Engine/Plugins/Experimental/NeuralNetworkInference/Source/NeuralNetworkInference/Internal/NeuralOperators/ConvBaseOperator.h b/Engine/Plugins/Experimental/NeuralNetworkInference/Source/NeuralNetworkInference/Internal/NeuralOperators/ConvBaseOperator.h
--- a/Engine/Plugins/Experimental/NeuralNetworkInference/Source/NeuralNetworkInference/Internal/NeuralOperators/ConvBaseOperator.h
+++ b/Engine/Plugins/Experimental/NeuralNetworkInference/Source/NeuralNetworkInference/Internal/NeuralOperators/ConvBaseOperator.h
@@ -121,7 +121,7 @@ private:
 
 
 
-/* FConvBaseOperator inline functions
+/* FConvBaseOperator inlined and templated functions
  *****************************************************************************/
 
 bool FConvBaseOperator::SetAndConfigureAuxiliaryTensor()
diff --git a/Engine/Plugins/Experimental/NeuralNetworkInference/Source/NeuralNetworkInference/Internal/NeuralOperators/ConvTransposeOperator.h b/Engine/Plugins/Experimental/NeuralNetworkInference/Source/NeuralNetworkInference/Internal/NeuralOperators/ConvTransposeOperator.h
--- a/Engine/Plugins/Experimental/NeuralNetworkInference/Source/NeuralNetworkInference/Internal/NeuralOperators/ConvTransposeOperator.h
+++ b/Engine/Plugins/Experimental/NeuralNetworkInference/Source/NeuralNetworkInference/Internal/NeuralOperators/ConvTransposeOperator.h
@@ -55,7 +55,7 @@ private:
 
 
 
-/* FConvBaseOperator inline functions
+/* FConvBaseOperator inlined and templated functions
  *****************************************************************************/
 
 void FConvTransposeOperator::SetWhetherWeightTensorWasFlipped(const bool bInIsWeightTensorFlipped)
diff --git a/Engine/Plugins/Experimental/NeuralNetworkInference/Source/NeuralNetworkInference/Internal/NeuralOperators/CosOperator.h b/Engine/Plugins/Experimental/NeuralNetworkInference/Source/NeuralNetworkInference/Internal/NeuralOperators/CosOperator.h
--- a/Engine/Plugins/Experimental/NeuralNetworkInference/Source/NeuralNetworkInference/Internal/NeuralOperators/CosOperator.h
+++ b/Engine/Plugins/Experimental/NeuralNetworkInference/Source/NeuralNetworkInference/Internal/NeuralOperators/CosOperator.h
@@ -17,7 +17,7 @@ public:
 
 
 
-/* FCosOperator inline functions
+/* FCosOperator inlined and templated functions
  *****************************************************************************/
 
 void FCosOperator::ForwardCPU()
diff --git a/Engine/Plugins/Experimental/NeuralNetworkInference/Source/NeuralNetworkInference/Internal/NeuralOperators/DivOperator.h b/Engine/Plugins/Experimental/NeuralNetworkInference/Source/NeuralNetworkInference/Internal/NeuralOperators/DivOperator.h
--- a/Engine/Plugins/Experimental/NeuralNetworkInference/Source/NeuralNetworkInference/Internal/NeuralOperators/DivOperator.h
+++ b/Engine/Plugins/Experimental/NeuralNetworkInference/Source/NeuralNetworkInference/Internal/NeuralOperators/DivOperator.h
@@ -17,7 +17,7 @@ public:
 
 
 
-/* FDivOperator inline functions
+/* FDivOperator inlined and templated functions
  *****************************************************************************/
 
 void FDivOperator::ForwardCPU()
diff --git a/Engine/Plugins/Experimental/NeuralNetworkInference/Source/NeuralNetworkInference/Internal/NeuralOperators/ExpOperator.h b/Engine/Plugins/Experimental/NeuralNetworkInference/Source/NeuralNetworkInference/Internal/NeuralOperators/ExpOperator.h
--- a/Engine/Plugins/Experimental/NeuralNetworkInference/Source/NeuralNetworkInference/Internal/NeuralOperators/ExpOperator.h
+++ b/Engine/Plugins/Experimental/NeuralNetworkInference/Source/NeuralNetworkInference/Internal/NeuralOperators/ExpOperator.h
@@ -17,7 +17,7 @@ public:
 
 
 
-/* FExpOperator inline functions
+/* FExpOperator inlined and templated functions
  *****************************************************************************/
 
 void FExpOperator::ForwardCPU()
diff --git a/Engine/Plugins/Experimental/NeuralNetworkInference/Source/NeuralNetworkInference/Internal/NeuralOperators/FloorOperator.h b/Engine/Plugins/Experimental/NeuralNetworkInference/Source/NeuralNetworkInference/Internal/NeuralOperators/FloorOperator.h
--- a/Engine/Plugins/Experimental/NeuralNetworkInference/Source/NeuralNetworkInference/Internal/NeuralOperators/FloorOperator.h
+++ b/Engine/Plugins/Experimental/NeuralNetworkInference/Source/NeuralNetworkInference/Internal/NeuralOperators/FloorOperator.h
@@ -17,7 +17,7 @@ public:
 
 
 
-/* FFloorOperator inline functions
+/* FFloorOperator inlined and templated functions
  *****************************************************************************/
 
 void FFloorOperator::ForwardCPU()
diff --git a/Engine/Plugins/Experimental/NeuralNetworkInference/Source/NeuralNetworkInference/Internal/NeuralOperators/LeakyReluOperator.h b/Engine/Plugins/Experimental/NeuralNetworkInference/Source/NeuralNetworkInference/Internal/NeuralOperators/LeakyReluOperator.h
--- a/Engine/Plugins/Experimental/NeuralNetworkInference/Source/NeuralNetworkInference/Internal/NeuralOperators/LeakyReluOperator.h
+++ b/Engine/Plugins/Experimental/NeuralNetworkInference/Source/NeuralNetworkInference/Internal/NeuralOperators/LeakyReluOperator.h
@@ -19,7 +19,7 @@ public:
 
 
 
-/* FLeakyReluOperator inline functions
+/* FLeakyReluOperator inlined and templated functions
  *****************************************************************************/
 
 void FLeakyReluOperator::ForwardCPU()
diff --git a/Engine/Plugins/Experimental/NeuralNetworkInference/Source/NeuralNetworkInference/Internal/NeuralOperators/LogOperator.h b/Engine/Plugins/Experimental/NeuralNetworkInference/Source/NeuralNetworkInference/Internal/NeuralOperators/LogOperator.h
--- a/Engine/Plugins/Experimental/NeuralNetworkInference/Source/NeuralNetworkInference/Internal/NeuralOperators/LogOperator.h
+++ b/Engine/Plugins/Experimental/NeuralNetworkInference/Source/NeuralNetworkInference/Internal/NeuralOperators/LogOperator.h
@@ -17,7 +17,7 @@ public:
 
 
 
-/* FLogOperator inline functions
+/* FLogOperator inlined and templated functions
  *****************************************************************************/
 
 void FLogOperator::ForwardCPU()
diff --git a/Engine/Plugins/Experimental/NeuralNetworkInference/Source/NeuralNetworkInference/Internal/NeuralOperators/MulOperator.h b/Engine/Plugins/Experimental/NeuralNetworkInference/Source/NeuralNetworkInference/Internal/NeuralOperators/MulOperator.h
--- a/Engine/Plugins/Experimental/NeuralNetworkInference/Source/NeuralNetworkInference/Internal/NeuralOperators/MulOperator.h
+++ b/Engine/Plugins/Experimental/NeuralNetworkInference/Source/NeuralNetworkInference/Internal/NeuralOperators/MulOperator.h
@@ -17,7 +17,7 @@ public:
 
 
 
-/* FMulOperator inline functions
+/* FMulOperator inlined and templated functions
  *****************************************************************************/
 
 void FMulOperator::ForwardCPU()
diff --git a/Engine/Plugins/Experimental/NeuralNetworkInference/Source/NeuralNetworkInference/Internal/NeuralOperators/NegOperator.h b/Engine/Plugins/Experimental/NeuralNetworkInference/Source/NeuralNetworkInference/Internal/NeuralOperators/NegOperator.h
--- a/Engine/Plugins/Experimental/NeuralNetworkInference/Source/NeuralNetworkInference/Internal/NeuralOperators/NegOperator.h
+++ b/Engine/Plugins/Experimental/NeuralNetworkInference/Source/NeuralNetworkInference/Internal/NeuralOperators/NegOperator.h
@@ -17,7 +17,7 @@ public:
 
 
 
-/* FNegOperator inline functions
+/* FNegOperator inlined and templated functions
  *****************************************************************************/
 
 void FNegOperator::ForwardCPU()
diff --git a/Engine/Plugins/Experimental/NeuralNetworkInference/Source/NeuralNetworkInference/Internal/NeuralOperators/PowOperator.h b/Engine/Plugins/Experimental/NeuralNetworkInference/Source/NeuralNetworkInference/Internal/NeuralOperators/PowOperator.h
--- a/Engine/Plugins/Experimental/NeuralNetworkInference/Source/NeuralNetworkInference/Internal/NeuralOperators/PowOperator.h
+++ b/Engine/Plugins/Experimental/NeuralNetworkInference/Source/NeuralNetworkInference/Internal/NeuralOperators/PowOperator.h
@@ -17,7 +17,7 @@ public:
 
 
 
-/* FPowOperator inline functions
+/* FPowOperator inlined and templated functions
  *****************************************************************************/
 
 void FPowOperator::ForwardCPU()
diff --git a/Engine/Plugins/Experimental/NeuralNetworkInference/Source/NeuralNetworkInference/Internal/NeuralOperators/ReciprocalOperator.h b/Engine/Plugins/Experimental/NeuralNetworkInference/Source/NeuralNetworkInference/Internal/NeuralOperators/ReciprocalOperator.h
--- a/Engine/Plugins/Experimental/NeuralNetworkInference/Source/NeuralNetworkInference/Internal/NeuralOperators/ReciprocalOperator.h
+++ b/Engine/Plugins/Experimental/NeuralNetworkInference/Source/NeuralNetworkInference/Internal/NeuralOperators/ReciprocalOperator.h
@@ -17,7 +17,7 @@ public:
 
 
 
-/* FReciprocalOperator inline functions
+/* FReciprocalOperator inlined and templated functions
  *****************************************************************************/
 
 void FReciprocalOperator::ForwardCPU()
diff --git a/Engine/Plugins/Experimental/NeuralNetworkInference/Source/NeuralNetworkInference/Internal/NeuralOperators/ReluOperator.h b/Engine/Plugins/Experimental/NeuralNetworkInference/Source/NeuralNetworkInference/Internal/NeuralOperators/ReluOperator.h
--- a/Engine/Plugins/Experimental/NeuralNetworkInference/Source/NeuralNetworkInference/Internal/NeuralOperators/ReluOperator.h
+++ b/Engine/Plugins/Experimental/NeuralNetworkInference/Source/NeuralNetworkInference/Internal/NeuralOperators/ReluOperator.h
@@ -17,7 +17,7 @@ public:
 
 
 
-/* FReluOperator inline functions
+/* FReluOperator inlined and templated functions
  *****************************************************************************/
 
 void FReluOperator::ForwardCPU()
diff --git a/Engine/Plugins/Experimental/NeuralNetworkInference/Source/NeuralNetworkInference/Internal/NeuralOperators/ReshapeOperator.h b/Engine/Plugins/Experimental/NeuralNetworkInference/Source/NeuralNetworkInference/Internal/NeuralOperators/ReshapeOperator.h
--- a/Engine/Plugins/Experimental/NeuralNetworkInference/Source/NeuralNetworkInference/Internal/NeuralOperators/ReshapeOperator.h
+++ b/Engine/Plugins/Experimental/NeuralNetworkInference/Source/NeuralNetworkInference/Internal/NeuralOperators/ReshapeOperator.h
@@ -52,7 +52,7 @@ private:
 
 
 
-/* FNeuralOperator inline functions
+/* FNeuralOperator inlined and templated functions
  *****************************************************************************/
 
 bool FReshapeOperator::NeedsPostForward() const
diff --git a/Engine/Plugins/Experimental/NeuralNetworkInference/Source/NeuralNetworkInference/Internal/NeuralOperators/RoundOperator.h b/Engine/Plugins/Experimental/NeuralNetworkInference/Source/NeuralNetworkInference/Internal/NeuralOperators/RoundOperator.h
--- a/Engine/Plugins/Experimental/NeuralNetworkInference/Source/NeuralNetworkInference/Internal/NeuralOperators/RoundOperator.h
+++ b/Engine/Plugins/Experimental/NeuralNetworkInference/Source/NeuralNetworkInference/Internal/NeuralOperators/RoundOperator.h
@@ -17,7 +17,7 @@ public:
 
 
 
-/* FRoundOperator inline functions
+/* FRoundOperator inlined and templated functions
  *****************************************************************************/
 
 void FRoundOperator::ForwardCPU()
diff --git a/Engine/Plugins/Experimental/NeuralNetworkInference/Source/NeuralNetworkInference/Internal/NeuralOperators/SigmoidOperator.h b/Engine/Plugins/Experimental/NeuralNetworkInference/Source/NeuralNetworkInference/Internal/NeuralOperators/SigmoidOperator.h
--- a/Engine/Plugins/Experimental/NeuralNetworkInference/Source/NeuralNetworkInference/Internal/NeuralOperators/SigmoidOperator.h
+++ b/Engine/Plugins/Experimental/NeuralNetworkInference/Source/NeuralNetworkInference/Internal/NeuralOperators/SigmoidOperator.h
@@ -17,7 +17,7 @@ public:
 
 
 
-/* FSigmoidOperator inline functions
+/* FSigmoidOperator inlined and templated functions
  *****************************************************************************/
 
 void FSigmoidOperator::ForwardCPU()
diff --git a/Engine/Plugins/Experimental/NeuralNetworkInference/Source/NeuralNetworkInference/Internal/NeuralOperators/SignOperator.h b/Engine/Plugins/Experimental/NeuralNetworkInference/Source/NeuralNetworkInference/Internal/NeuralOperators/SignOperator.h
--- a/Engine/Plugins/Experimental/NeuralNetworkInference/Source/NeuralNetworkInference/Internal/NeuralOperators/SignOperator.h
+++ b/Engine/Plugins/Experimental/NeuralNetworkInference/Source/NeuralNetworkInference/Internal/NeuralOperators/SignOperator.h
@@ -17,7 +17,7 @@ public:
 
 
 
-/* FSignOperator inline functions
+/* FSignOperator inlined and templated functions
  *****************************************************************************/
 
 void FSignOperator::ForwardCPU()
diff --git a/Engine/Plugins/Experimental/NeuralNetworkInference/Source/NeuralNetworkInference/Internal/NeuralOperators/SinOperator.h b/Engine/Plugins/Experimental/NeuralNetworkInference/Source/NeuralNetworkInference/Internal/NeuralOperators/SinOperator.h
--- a/Engine/Plugins/Experimental/NeuralNetworkInference/Source/NeuralNetworkInference/Internal/NeuralOperators/SinOperator.h
+++ b/Engine/Plugins/Experimental/NeuralNetworkInference/Source/NeuralNetworkInference/Internal/NeuralOperators/SinOperator.h
@@ -17,7 +17,7 @@ public:
 
 
 
-/* FSinOperator inline functions
+/* FSinOperator inlined and templated functions
  *****************************************************************************/
 
 void FSinOperator::ForwardCPU()
diff --git a/Engine/Plugins/Experimental/NeuralNetworkInference/Source/NeuralNetworkInference/Internal/NeuralOperators/SinhOperator.h b/Engine/Plugins/Experimental/NeuralNetworkInference/Source/NeuralNetworkInference/Internal/NeuralOperators/SinhOperator.h
--- a/Engine/Plugins/Experimental/NeuralNetworkInference/Source/NeuralNetworkInference/Internal/NeuralOperators/SinhOperator.h
+++ b/Engine/Plugins/Experimental/NeuralNetworkInference/Source/NeuralNetworkInference/Internal/NeuralOperators/SinhOperator.h
@@ -17,7 +17,7 @@ public:
 
 
 
-/* FSinhOperator inline functions
+/* FSinhOperator inlined and templated functions
  *****************************************************************************/
 
 void FSinhOperator::ForwardCPU()
diff --git a/Engine/Plugins/Experimental/NeuralNetworkInference/Source/NeuralNetworkInference/Internal/NeuralOperators/SqrtOperator.h b/Engine/Plugins/Experimental/NeuralNetworkInference/Source/NeuralNetworkInference/Internal/NeuralOperators/SqrtOperator.h
--- a/Engine/Plugins/Experimental/NeuralNetworkInference/Source/NeuralNetworkInference/Internal/NeuralOperators/SqrtOperator.h
+++ b/Engine/Plugins/Experimental/NeuralNetworkInference/Source/NeuralNetworkInference/Internal/NeuralOperators/SqrtOperator.h
@@ -17,7 +17,7 @@ public:
 
 
 
-/* FSqrtOperator inline functions
+/* FSqrtOperator inlined and templated functions
  *****************************************************************************/
 
 void FSqrtOperator::ForwardCPU()
diff --git a/Engine/Plugins/Experimental/NeuralNetworkInference/Source/NeuralNetworkInference/Internal/NeuralOperators/SqueezeOperator.h b/Engine/Plugins/Experimental/NeuralNetworkInference/Source/NeuralNetworkInference/Internal/NeuralOperators/SqueezeOperator.h
--- a/Engine/Plugins/Experimental/NeuralNetworkInference/Source/NeuralNetworkInference/Internal/NeuralOperators/SqueezeOperator.h
+++ b/Engine/Plugins/Experimental/NeuralNetworkInference/Source/NeuralNetworkInference/Internal/NeuralOperators/SqueezeOperator.h
@@ -46,7 +46,7 @@ private:
 
 
 
-/* FNeuralOperator inline functions
+/* FNeuralOperator inlined and templated functions
  *****************************************************************************/
 
 bool FSqueezeOperator::NeedsPostForward() const
diff --git a/Engine/Plugins/Experimental/NeuralNetworkInference/Source/NeuralNetworkInference/Internal/NeuralOperators/SubOperator.h b/Engine/Plugins/Experimental/NeuralNetworkInference/Source/NeuralNetworkInference/Internal/NeuralOperators/SubOperator.h
--- a/Engine/Plugins/Experimental/NeuralNetworkInference/Source/NeuralNetworkInference/Internal/NeuralOperators/SubOperator.h
+++ b/Engine/Plugins/Experimental/NeuralNetworkInference/Source/NeuralNetworkInference/Internal/NeuralOperators/SubOperator.h
@@ -17,7 +17,7 @@ public:
 
 
 
-/* FSubOperator inline functions
+/* FSubOperator inlined and templated functions
  *****************************************************************************/
 
 void FSubOperator::ForwardCPU()
diff --git a/Engine/Plugins/Experimental/NeuralNetworkInference/Source/NeuralNetworkInference/Internal/NeuralOperators/TanOperator.h b/Engine/Plugins/Experimental/NeuralNetworkInference/Source/NeuralNetworkInference/Internal/NeuralOperators/TanOperator.h
--- a/Engine/Plugins/Experimental/NeuralNetworkInference/Source/NeuralNetworkInference/Internal/NeuralOperators/TanOperator.h
+++ b/Engine/Plugins/Experimental/NeuralNetworkInference/Source/NeuralNetworkInference/Internal/NeuralOperators/TanOperator.h
@@ -17,7 +17,7 @@ public:
 
 
 
-/* FTanOperator inline functions
+/* FTanOperator inlined and templated functions
  *****************************************************************************/
 
 void FTanOperator::ForwardCPU()
diff --git a/Engine/Plugins/Experimental/NeuralNetworkInference/Source/NeuralNetworkInference/Private/UEOnly/NeuralTensorManager.h b/Engine/Plugins/Experimental/NeuralNetworkInference/Source/NeuralNetworkInference/Private/UEOnly/NeuralTensorManager.h
--- a/Engine/Plugins/Experimental/NeuralNetworkInference/Source/NeuralNetworkInference/Private/UEOnly/NeuralTensorManager.h
+++ b/Engine/Plugins/Experimental/NeuralNetworkInference/Source/NeuralNetworkInference/Private/UEOnly/NeuralTensorManager.h
@@ -173,7 +173,7 @@ protected:
 
 
 
-/* FNeuralTensorManager inline functions
+/* FNeuralTensorManager inlined and templated functions
  *****************************************************************************/
 
 template<typename T>
diff --git a/Engine/Plugins/Experimental/NeuralNetworkInference/Source/NeuralNetworkInferenceQA/Private/OperatorTester.h b/Engine/Plugins/Experimental/NeuralNetworkInference/Source/NeuralNetworkInferenceQA/Private/OperatorTester.h
--- a/Engine/Plugins/Experimental/NeuralNetworkInference/Source/NeuralNetworkInferenceQA/Private/OperatorTester.h
+++ b/Engine/Plugins/Experimental/NeuralNetworkInference/Source/NeuralNetworkInferenceQA/Private/OperatorTester.h
@@ -80,7 +80,7 @@ private:
 
 
 
-/* FOperatorTester inline functions
+/* FOperatorTester inlined and templated functions
  *****************************************************************************/
 
 template<typename TNeuralOperatorChild>
