Adding NNI documentation - Part 4:
- NNI/Public/NeuraTensor.h and NNIEditor/INeuralNetworkInferenceEditorModule with class doc

#rnx

diff --git a/Engine/Plugins/Experimental/NeuralNetworkInference/Source/NeuralNetworkInference/Public/NeuralTensor.h b/Engine/Plugins/Experimental/NeuralNetworkInference/Source/NeuralNetworkInference/Public/NeuralTensor.h
--- a/Engine/Plugins/Experimental/NeuralNetworkInference/Source/NeuralNetworkInference/Public/NeuralTensor.h
+++ b/Engine/Plugins/Experimental/NeuralNetworkInference/Source/NeuralNetworkInference/Public/NeuralTensor.h
@@ -13,12 +13,13 @@
 #include "NeuralTensor.generated.h"
 
 /**
- * This is an auxiliary class. See UNeuralNetwork for a high-level wrapper of the whole NeuralNetworkInference plugin. The UNeuralNetwork header
- * documentation also includes some code examples.
+ * @see UNeuralNetwork for general documentation and code samples about the whole NeuralNetworkInference (NNI) plugin.
  *
- * FNeuralTensor represents a tensor of the UNeuralNetwork model.
- * Its CPU operations are very similar to those of TArray<T>. Most of its operations run on CPU, so `ToGPU_RenderThread()` must be called before
- * running on GPU and after running any FNeuralTensor function that modified the CPU memory.
+ * FNeuralTensor is an auxiliary class of UNeuralNetwork which represents a tensor of the UNeuralNetwork model. It is Unreal Engine's equivalent of
+ * torch.Tensor (PyTorch) or caffe::Blob.
+ *
+ * Most of its functions run on the CPU, so `ToGPU_RenderThread()` must be called before running on GPU and after running any FNeuralTensor function
+ * that modifies the CPU memory. In addition, FNeuralTensor's CPU functions are very similar to those of TArray<T>.
  */
 USTRUCT()
 struct NEURALNETWORKINFERENCE_API FNeuralTensor
diff --git a/Engine/Plugins/Experimental/NeuralNetworkInference/Source/NeuralNetworkInferenceEditor/Public/NeuralNetworkInferenceEditorModule.h b/Engine/Plugins/Experimental/NeuralNetworkInference/Source/NeuralNetworkInferenceEditor/Public/NeuralNetworkInferenceEditorModule.h
--- a/Engine/Plugins/Experimental/NeuralNetworkInference/Source/NeuralNetworkInferenceEditor/Public/NeuralNetworkInferenceEditorModule.h
+++ b/Engine/Plugins/Experimental/NeuralNetworkInference/Source/NeuralNetworkInferenceEditor/Public/NeuralNetworkInferenceEditorModule.h
@@ -6,11 +6,22 @@
 #include "AssetTypeCategories.h"
 #include "Modules/ModuleManager.h"
 
+/**
+ * @see UNeuralNetwork for general documentation and code samples about the whole NeuralNetworkInference (NNI) plugin.
+ *
+ * General NNI users do not need to use this class nor any other class from the NeuralNetworkInferenceEditor (NNIEditor) module, except for
+ * potentially INeuralNetworkInferenceEditorModule::GetMLAssetCategoryBit().
+ *
+ * INeuralNetworkInferenceEditorModule is the IModuleInterface of NeuralNetworkInferenceEditor (NNIEditor), and additionally provides the public
+ * function INeuralNetworkInferenceEditorModule::GetMLAssetCategoryBit().
+ */
 class NEURALNETWORKINFERENCEEDITOR_API INeuralNetworkInferenceEditorModule : public IModuleInterface
 {
 public:
 	/**
-	 * It returns the EAssetTypeCategories::Type for the "Machine Learning" category
+	 * It returns the EAssetTypeCategories::Type for the "Machine Learning" category. This allows for non-NNI classes to appear in the
+	 * "Machine Learning" category. To see all classes on this category, right-click on the Editor's "Content Browser", then "Create Advanced Asset",
+	 * and "Machine Learning". "Neural Network" and any other class registered with GetMLAssetCategoryBit() should appear in there.
 	 */
 	virtual EAssetTypeCategories::Type GetMLAssetCategoryBit() const = 0;
 };
