Adding NNI documentation:
- Moved NNIProfiler and NNIQA into "Private" (as the client do not need them).

#rnx

diff --git a/Engine/Plugins/Experimental/NeuralNetworkInference/Source/NeuralNetworkInferenceProfiling/NeuralNetworkInferenceProfiling.Build.cs b/Engine/Plugins/Experimental/NeuralNetworkInference/Source/NeuralNetworkInferenceProfiling/NeuralNetworkInferenceProfiling.Build.cs
deleted file mode 100644
diff --git a/Engine/Plugins/Experimental/NeuralNetworkInference/Source/NeuralNetworkInferenceProfiling/Private/Module.cpp b/Engine/Plugins/Experimental/NeuralNetworkInference/Source/NeuralNetworkInferenceProfiling/Private/Module.cpp
deleted file mode 100644
diff --git a/Engine/Plugins/Experimental/NeuralNetworkInference/Source/NeuralNetworkInferenceProfiling/Private/Module.h b/Engine/Plugins/Experimental/NeuralNetworkInference/Source/NeuralNetworkInferenceProfiling/Private/Module.h
deleted file mode 100644
diff --git a/Engine/Plugins/Experimental/NeuralNetworkInference/Source/NeuralNetworkInferenceProfiling/Private/NeuralStats.cpp b/Engine/Plugins/Experimental/NeuralNetworkInference/Source/NeuralNetworkInferenceProfiling/Private/NeuralStats.cpp
deleted file mode 100644
diff --git a/Engine/Plugins/Experimental/NeuralNetworkInference/Source/NeuralNetworkInferenceProfiling/Private/NeuralTimer.cpp b/Engine/Plugins/Experimental/NeuralNetworkInference/Source/NeuralNetworkInferenceProfiling/Private/NeuralTimer.cpp
deleted file mode 100644
diff --git a/Engine/Plugins/Experimental/NeuralNetworkInference/Source/NeuralNetworkInferenceProfiling/Public/NeuralStats.h b/Engine/Plugins/Experimental/NeuralNetworkInference/Source/NeuralNetworkInferenceProfiling/Public/NeuralStats.h
deleted file mode 100644
diff --git a/Engine/Plugins/Experimental/NeuralNetworkInference/Source/NeuralNetworkInferenceProfiling/Public/NeuralTimer.h b/Engine/Plugins/Experimental/NeuralNetworkInference/Source/NeuralNetworkInferenceProfiling/Public/NeuralTimer.h
deleted file mode 100644
diff --git a/Engine/Plugins/Experimental/NeuralNetworkInference/Source/NeuralNetworkInferenceQA/Internal/NeuralNetworkInferenceQAAsset.h b/Engine/Plugins/Experimental/NeuralNetworkInference/Source/NeuralNetworkInferenceQA/Internal/NeuralNetworkInferenceQAAsset.h
deleted file mode 100644
diff --git a/Engine/Plugins/Experimental/NeuralNetworkInference/Source/NeuralNetworkInferenceQA/NeuralNetworkInferenceQA.Build.cs b/Engine/Plugins/Experimental/NeuralNetworkInference/Source/NeuralNetworkInferenceQA/NeuralNetworkInferenceQA.Build.cs
deleted file mode 100644
diff --git a/Engine/Plugins/Experimental/NeuralNetworkInference/Source/NeuralNetworkInferenceQA/Private/ModelProtoFileReaderTester.cpp b/Engine/Plugins/Experimental/NeuralNetworkInference/Source/NeuralNetworkInferenceQA/Private/ModelProtoFileReaderTester.cpp
deleted file mode 100644
diff --git a/Engine/Plugins/Experimental/NeuralNetworkInference/Source/NeuralNetworkInferenceQA/Private/ModelProtoFileReaderTester.h b/Engine/Plugins/Experimental/NeuralNetworkInference/Source/NeuralNetworkInferenceQA/Private/ModelProtoFileReaderTester.h
deleted file mode 100644
diff --git a/Engine/Plugins/Experimental/NeuralNetworkInference/Source/NeuralNetworkInferenceQA/Private/ModelUnitTester.cpp b/Engine/Plugins/Experimental/NeuralNetworkInference/Source/NeuralNetworkInferenceQA/Private/ModelUnitTester.cpp
deleted file mode 100644
diff --git a/Engine/Plugins/Experimental/NeuralNetworkInference/Source/NeuralNetworkInferenceQA/Private/ModelUnitTester.h b/Engine/Plugins/Experimental/NeuralNetworkInference/Source/NeuralNetworkInferenceQA/Private/ModelUnitTester.h
deleted file mode 100644
diff --git a/Engine/Plugins/Experimental/NeuralNetworkInference/Source/NeuralNetworkInferenceQA/Private/Module.cpp b/Engine/Plugins/Experimental/NeuralNetworkInference/Source/NeuralNetworkInferenceQA/Private/Module.cpp
deleted file mode 100644
diff --git a/Engine/Plugins/Experimental/NeuralNetworkInference/Source/NeuralNetworkInferenceQA/Private/Module.h b/Engine/Plugins/Experimental/NeuralNetworkInference/Source/NeuralNetworkInferenceQA/Private/Module.h
deleted file mode 100644
diff --git a/Engine/Plugins/Experimental/NeuralNetworkInference/Source/NeuralNetworkInferenceQA/Private/NeuralNetworkInferenceQA.cpp b/Engine/Plugins/Experimental/NeuralNetworkInference/Source/NeuralNetworkInferenceQA/Private/NeuralNetworkInferenceQA.cpp
deleted file mode 100644
diff --git a/Engine/Plugins/Experimental/NeuralNetworkInference/Source/NeuralNetworkInferenceQA/Private/NeuralNetworkInferenceQAAsset.cpp b/Engine/Plugins/Experimental/NeuralNetworkInference/Source/NeuralNetworkInferenceQA/Private/NeuralNetworkInferenceQAAsset.cpp
deleted file mode 100644
diff --git a/Engine/Plugins/Experimental/NeuralNetworkInference/Source/NeuralNetworkInferenceQA/Private/NeuralNetworkInferenceQAUtils.cpp b/Engine/Plugins/Experimental/NeuralNetworkInference/Source/NeuralNetworkInferenceQA/Private/NeuralNetworkInferenceQAUtils.cpp
deleted file mode 100644
diff --git a/Engine/Plugins/Experimental/NeuralNetworkInference/Source/NeuralNetworkInferenceQA/Private/NeuralNetworkInferenceQAUtils.h b/Engine/Plugins/Experimental/NeuralNetworkInference/Source/NeuralNetworkInferenceQA/Private/NeuralNetworkInferenceQAUtils.h
deleted file mode 100644
diff --git a/Engine/Plugins/Experimental/NeuralNetworkInference/Source/NeuralNetworkInferenceQA/Private/OperatorTester.cpp b/Engine/Plugins/Experimental/NeuralNetworkInference/Source/NeuralNetworkInferenceQA/Private/OperatorTester.cpp
deleted file mode 100644
diff --git a/Engine/Plugins/Experimental/NeuralNetworkInference/Source/NeuralNetworkInferenceQA/Private/OperatorTester.h b/Engine/Plugins/Experimental/NeuralNetworkInference/Source/NeuralNetworkInferenceQA/Private/OperatorTester.h
deleted file mode 100644
diff --git a/Engine/Plugins/Experimental/NeuralNetworkInference/Source/NeuralNetworkInferenceQA/Private/OperatorUnitTester.cpp b/Engine/Plugins/Experimental/NeuralNetworkInference/Source/NeuralNetworkInferenceQA/Private/OperatorUnitTester.cpp
deleted file mode 100644
diff --git a/Engine/Plugins/Experimental/NeuralNetworkInference/Source/NeuralNetworkInferenceQA/Private/OperatorUnitTester.h b/Engine/Plugins/Experimental/NeuralNetworkInference/Source/NeuralNetworkInferenceQA/Private/OperatorUnitTester.h
deleted file mode 100644
diff --git a/Engine/Plugins/Experimental/NeuralNetworkInference/Source/NeuralNetworkInferenceQA/Private/UnitTester.cpp b/Engine/Plugins/Experimental/NeuralNetworkInference/Source/NeuralNetworkInferenceQA/Private/UnitTester.cpp
deleted file mode 100644
diff --git a/Engine/Plugins/Experimental/NeuralNetworkInference/Source/NeuralNetworkInferenceQA/Private/UnitTester.h b/Engine/Plugins/Experimental/NeuralNetworkInference/Source/NeuralNetworkInferenceQA/Private/UnitTester.h
deleted file mode 100644
diff --git a/Engine/Plugins/Experimental/NeuralNetworkInference/Source/NeuralNetworkInferenceQA/Public/NeuralNetworkInferenceQA.h b/Engine/Plugins/Experimental/NeuralNetworkInference/Source/NeuralNetworkInferenceQA/Public/NeuralNetworkInferenceQA.h
deleted file mode 100644
diff --git a/Engine/Plugins/Experimental/NeuralNetworkInference/Source/Private/NeuralNetworkInferenceProfiling/NeuralNetworkInferenceProfiling.Build.cs b/Engine/Plugins/Experimental/NeuralNetworkInference/Source/Private/NeuralNetworkInferenceProfiling/NeuralNetworkInferenceProfiling.Build.cs
new file mode 100644
--- /dev/null
+++ b/Engine/Plugins/Experimental/NeuralNetworkInference/Source/Private/NeuralNetworkInferenceProfiling/NeuralNetworkInferenceProfiling.Build.cs
@@ -0,0 +1,38 @@
+// Copyright Epic Games, Inc. All Rights Reserved.
+
+using UnrealBuildTool;
+using System.IO;
+
+public class NeuralNetworkInferenceProfiling : ModuleRules
+{
+	public NeuralNetworkInferenceProfiling(ReadOnlyTargetRules Target) : base(Target)
+	{
+
+		ShortName = "NNIProfiling"; 
+		PCHUsage = ModuleRules.PCHUsageMode.UseExplicitOrSharedPCHs;
+
+		PublicIncludePaths.AddRange(
+			new string[] {
+				System.IO.Path.Combine(ModuleDirectory, "..")
+			}
+		);
+
+		PublicDependencyModuleNames.AddRange
+			(
+			new string[] {
+				"Core",
+				"Engine",
+			}
+		);
+
+		PrivateDependencyModuleNames.AddRange
+			(
+			new string[] {
+				"ThirdPartyHelperAndDLLLoader",
+				"Boost"
+			}
+		);
+
+		PublicDefinitions.Add("WITH_NNI_STATS");
+	}
+}
diff --git a/Engine/Plugins/Experimental/NeuralNetworkInference/Source/Private/NeuralNetworkInferenceProfiling/Private/Module.cpp b/Engine/Plugins/Experimental/NeuralNetworkInference/Source/Private/NeuralNetworkInferenceProfiling/Private/Module.cpp
new file mode 100644
--- /dev/null
+++ b/Engine/Plugins/Experimental/NeuralNetworkInference/Source/Private/NeuralNetworkInferenceProfiling/Private/Module.cpp
@@ -0,0 +1,16 @@
+// Copyright Epic Games, Inc. All Rights Reserved.
+
+#include "Module.h"
+
+// This code will execute after your module is loaded into memory; the exact timing is specified in the .uplugin file per-module
+void FNeuralNetworkInferenceProfilingModule::StartupModule()
+{
+}
+
+// This function may be called during shutdown to clean up your module. For modules that support dynamic reloading,
+// we call this function before unloading the module.
+void FNeuralNetworkInferenceProfilingModule::ShutdownModule()
+{
+}
+
+IMPLEMENT_MODULE(FNeuralNetworkInferenceProfilingModule, NeuralNetworkInferenceProfiling);
diff --git a/Engine/Plugins/Experimental/NeuralNetworkInference/Source/Private/NeuralNetworkInferenceProfiling/Private/Module.h b/Engine/Plugins/Experimental/NeuralNetworkInference/Source/Private/NeuralNetworkInferenceProfiling/Private/Module.h
new file mode 100644
--- /dev/null
+++ b/Engine/Plugins/Experimental/NeuralNetworkInference/Source/Private/NeuralNetworkInferenceProfiling/Private/Module.h
@@ -0,0 +1,13 @@
+// Copyright Epic Games, Inc. All Rights Reserved.
+
+#pragma once
+
+#include "Modules/ModuleManager.h"
+
+class FNeuralNetworkInferenceProfilingModule : public IModuleInterface
+{
+public:
+	// Begin IModuleInterface
+	virtual void StartupModule() override;
+	virtual void ShutdownModule() override;
+};
diff --git a/Engine/Plugins/Experimental/NeuralNetworkInference/Source/Private/NeuralNetworkInferenceProfiling/Private/NeuralStats.cpp b/Engine/Plugins/Experimental/NeuralNetworkInference/Source/Private/NeuralNetworkInferenceProfiling/Private/NeuralStats.cpp
new file mode 100644
--- /dev/null
+++ b/Engine/Plugins/Experimental/NeuralNetworkInference/Source/Private/NeuralNetworkInferenceProfiling/Private/NeuralStats.cpp
@@ -0,0 +1,123 @@
+// Copyright Epic Games, Inc. All Rights Reserved.
+
+#include "NeuralStats.h"
+#include "GenericPlatform/GenericPlatformMath.h"
+
+DEFINE_LOG_CATEGORY_STATIC(LogNNIStats, Display, All);
+
+
+
+/* FNeuralStats auxiliary functions
+ *****************************************************************************/
+
+void FNeuralStats_CalculateMinMax(float& OutMin, float& OutMax, const TCircularBuffer<float>& InDataBuffer, const bool bInIsBufferFull, const uint32 InBufferIdx)
+{
+	const uint32 PtrEnd = (bInIsBufferFull) ? InDataBuffer.Capacity() : InBufferIdx;
+
+	OutMin = FLT_MAX;
+	OutMax = FLT_MIN;
+	for (uint32 Index = 0; Index < PtrEnd; ++Index)
+	{
+		const float CurrentElement = InDataBuffer[Index];
+		if (CurrentElement < OutMin)
+		{
+			OutMin = CurrentElement;
+		}
+		if (CurrentElement > OutMax)
+		{
+			OutMax = CurrentElement;
+		}
+	}
+}
+
+
+
+/* FNeuralStatsData auxiliary functions
+ *****************************************************************************/
+
+FNeuralStatsData::FNeuralStatsData(const int InNumberSamples, const float InAverage, const float InStdDev, const float InMin, const float InMax)
+	: NumberSamples(InNumberSamples)
+	, Average(InAverage)
+	, StdDev(InStdDev)
+	, Min(InMin)
+	, Max(InMax)
+{}
+
+
+
+/* FNeuralNetworkInferenceQATimer public functions
+ *****************************************************************************/
+
+FNeuralStats::FNeuralStats(const uint32 InSizeRollingWindow) 
+	: EpsilonFloat(0.000001f)
+	, BufferIdx(0)
+	, bIsBufferFull(false)
+	, LastSample(0.f)
+	, DataBuffer(InSizeRollingWindow, 0)
+{
+}
+	
+void FNeuralStats::StoreSample(const float InRunTime) 
+{
+	LastSample = InRunTime;
+
+	DataBuffer[BufferIdx] = LastSample;
+
+	if(BufferIdx == DataBuffer.Capacity() && !bIsBufferFull)
+	{
+		bIsBufferFull = true;
+	}
+	BufferIdx = DataBuffer.GetNextIndex(BufferIdx);
+}
+
+void FNeuralStats::ResetStats()
+{
+	for (uint32 Index = 0; Index < DataBuffer.Capacity(); ++Index)
+	{
+		DataBuffer[Index] = 0;
+	}
+	bIsBufferFull = false;
+	BufferIdx = 0;
+	LastSample = 0;
+}
+
+float FNeuralStats::GetLastSample() const
+{
+	return LastSample;
+}
+
+FNeuralStatsData FNeuralStats::GetStats() const
+{
+	const uint32 count = (bIsBufferFull) ? DataBuffer.Capacity() : BufferIdx;
+	const float Mean = CalculateMean();
+	const float StdDev = CalculateStdDev(Mean);
+	float Min, Max;
+	FNeuralStats_CalculateMinMax(Min, Max, DataBuffer, bIsBufferFull, BufferIdx);
+	// Return aggregated data
+	return FNeuralStatsData(count, Mean, StdDev, Min, Max);
+}
+
+float FNeuralStats::CalculateMean() const
+{
+	uint32 PtrEnd = (bIsBufferFull) ? DataBuffer.Capacity(): BufferIdx;
+	float AccumulatorMean = 0;
+	for (uint32 Index = 0; Index < PtrEnd; ++Index)
+	{
+		AccumulatorMean += DataBuffer[Index];
+	}
+	// Return mean
+	return AccumulatorMean / (static_cast<float>(PtrEnd) + EpsilonFloat);
+}
+
+float FNeuralStats::CalculateStdDev(const float InMean) const
+{
+	const uint32 PtrEnd = (bIsBufferFull) ? DataBuffer.Capacity() : BufferIdx;
+
+	float AccumulatorStdDev = 0;
+	for (uint32 Index = 0; Index < PtrEnd; ++Index)
+	{
+		AccumulatorStdDev += (DataBuffer[Index] - InMean) * (DataBuffer[Index] - InMean);
+	}
+	// Return std deviation
+	return FMath::Sqrt(AccumulatorStdDev / (static_cast<float>(PtrEnd) + EpsilonFloat));
+}
diff --git a/Engine/Plugins/Experimental/NeuralNetworkInference/Source/Private/NeuralNetworkInferenceProfiling/Private/NeuralTimer.cpp b/Engine/Plugins/Experimental/NeuralNetworkInference/Source/Private/NeuralNetworkInferenceProfiling/Private/NeuralTimer.cpp
new file mode 100644
--- /dev/null
+++ b/Engine/Plugins/Experimental/NeuralNetworkInference/Source/Private/NeuralNetworkInferenceProfiling/Private/NeuralTimer.cpp
@@ -0,0 +1,17 @@
+// Copyright Epic Games, Inc. All Rights Reserved.
+
+#include "NeuralTimer.h"
+#include "HAL/PlatformTime.h"
+
+/* FNeuralNetworkInferenceQATimer public functions
+ *****************************************************************************/
+
+void FNeuralTimer::Tic()
+{
+	TimeStart = {FPlatformTime::Seconds()};
+}
+
+double FNeuralTimer::Toc() const
+{
+	return (FPlatformTime::Seconds() - TimeStart) * 1e3;
+}
diff --git a/Engine/Plugins/Experimental/NeuralNetworkInference/Source/Private/NeuralNetworkInferenceProfiling/Public/NeuralStats.h b/Engine/Plugins/Experimental/NeuralNetworkInference/Source/Private/NeuralNetworkInferenceProfiling/Public/NeuralStats.h
new file mode 100644
--- /dev/null
+++ b/Engine/Plugins/Experimental/NeuralNetworkInference/Source/Private/NeuralNetworkInferenceProfiling/Public/NeuralStats.h
@@ -0,0 +1,39 @@
+// Copyright Epic Games, Inc. All Rights Reserved.
+
+#pragma once
+
+#include "CoreMinimal.h"
+#include "Containers/CircularBuffer.h"
+
+struct NEURALNETWORKINFERENCEPROFILING_API FNeuralStatsData
+{
+	const uint32 NumberSamples;
+	const float Average;
+	const float StdDev;
+	const float Min;
+	const float Max;
+
+	FNeuralStatsData(const int InNumberSamples, const float InAverage, const float InStdDev, const float InMin, const float InMax);
+};
+
+class NEURALNETWORKINFERENCEPROFILING_API FNeuralStats
+{
+public:
+	FNeuralStats(const uint32 InSizeRollingWindow = 1024);
+	
+	void StoreSample(const float InRunTime);
+	void ResetStats();
+
+	float GetLastSample() const;
+	FNeuralStatsData GetStats() const;
+
+private:
+	const float EpsilonFloat;
+	uint32 BufferIdx;
+	bool bIsBufferFull;
+	float LastSample;
+	TCircularBuffer<float> DataBuffer;
+
+	float CalculateMean() const;
+	float CalculateStdDev(const float InMean) const;
+};
diff --git a/Engine/Plugins/Experimental/NeuralNetworkInference/Source/Private/NeuralNetworkInferenceProfiling/Public/NeuralTimer.h b/Engine/Plugins/Experimental/NeuralNetworkInference/Source/Private/NeuralNetworkInferenceProfiling/Public/NeuralTimer.h
new file mode 100644
--- /dev/null
+++ b/Engine/Plugins/Experimental/NeuralNetworkInference/Source/Private/NeuralNetworkInferenceProfiling/Public/NeuralTimer.h
@@ -0,0 +1,19 @@
+// Copyright Epic Games, Inc. All Rights Reserved.
+
+#pragma once
+
+#include "CoreMinimal.h"
+
+class NEURALNETWORKINFERENCEPROFILING_API FNeuralTimer
+{
+public:
+	void Tic();
+
+	/**
+	 * Time in milliseconds, but with nanosecond accuracy (e.g., 1.234567 msec).
+	 */
+	double Toc() const;
+
+private:
+	double TimeStart;
+};
diff --git a/Engine/Plugins/Experimental/NeuralNetworkInference/Source/Private/NeuralNetworkInferenceQA/Internal/NeuralNetworkInferenceQAAsset.h b/Engine/Plugins/Experimental/NeuralNetworkInference/Source/Private/NeuralNetworkInferenceQA/Internal/NeuralNetworkInferenceQAAsset.h
new file mode 100644
--- /dev/null
+++ b/Engine/Plugins/Experimental/NeuralNetworkInference/Source/Private/NeuralNetworkInferenceQA/Internal/NeuralNetworkInferenceQAAsset.h
@@ -0,0 +1,107 @@
+// Copyright Epic Games, Inc. All Rights Reserved.
+
+#pragma once
+
+#include "CoreMinimal.h"
+#include "NeuralNetwork.h"
+#include "NeuralTensor.h"
+#include "NeuralNetworkInferenceQAAsset.generated.h"
+
+
+
+USTRUCT(BlueprintType)
+struct NEURALNETWORKINFERENCEQA_API FNeuralNetworkInferenceQAOperatorTestAsset
+{
+	GENERATED_BODY()
+
+public:
+	UPROPERTY(VisibleAnywhere, Category = "Neural Network Inference QA")
+	TArray<FNeuralTensor> InputTensors;
+	
+	UPROPERTY(VisibleAnywhere, Category = "Neural Network Inference QA")
+	TArray<FNeuralTensor> OutputTensors;
+
+	/**
+	 * It adds the InNameIndexMap subset of InTensors into Tensors.
+	 */
+	void AddInputTensors(const UNeuralNetwork* const InNetwork);
+	void AddOutputTensors(const UNeuralNetwork* const InNetwork);
+
+	/**
+	 * It will compare new vs. previous tests.
+	 */
+	static bool CompareAverageL1DiffNewVsPreviousTests(const FNeuralNetworkInferenceQAOperatorTestAsset& InOperatorTestAsset1,
+		const FNeuralNetworkInferenceQAOperatorTestAsset& InOperatorTestAsset2, const FString& InOperatorName);
+
+private:
+	/**
+	 * Auxiliary for CompareAverageL1DiffNewVsPreviousTests(const FNeuralNetworkInferenceQAOperatorTestAsset&, const FNeuralNetworkInferenceQAOperatorTestAsset&).
+	 */
+	static bool CompareAverageL1DiffNewVsPreviousTests(const TArray<FNeuralTensor>& InTensorsA, const TArray<FNeuralTensor>& InTensorsB, const FString& InOperatorName);
+};
+
+
+
+USTRUCT(BlueprintType)
+struct NEURALNETWORKINFERENCEQA_API FNeuralNetworkInferenceQAOperatorAsset
+{
+	GENERATED_BODY()
+
+public:
+	static FString RunNetworkCPUAndGetString(UNeuralNetwork* InOutNetwork);
+
+	void RunAndAddTest(UNeuralNetwork* InOutNetwork);
+
+	/**
+	 * It will compare new vs. previous tests.
+	 */
+	bool CompareNewVsPreviousTests(const FString& InGroundTruthDirectory, const FString& InOperatorName);
+
+	/**
+	 * It flushes (cleans/resets) the new results so the tests can be run again.
+	 */
+	void FlushNewTests();
+	
+protected:
+	FString NewTestsString;
+
+	TArray<FNeuralNetworkInferenceQAOperatorTestAsset> NewTests;
+
+	UPROPERTY(VisibleAnywhere, Category = "Neural Network Inference QA")
+	TArray<FNeuralNetworkInferenceQAOperatorTestAsset> PreviousTests;
+
+	UPROPERTY()
+	FString PreviousTestsString;
+};
+
+
+
+/**
+ * This is an auxiliary class. See UNeuralNetwork for a high-level wrapper of the whole NeuralNetworkInference plugin. The UNeuralNetwork header
+ * documentation also includes some code examples.
+ */
+UCLASS(BlueprintType)
+class NEURALNETWORKINFERENCEQA_API UNeuralNetworkInferenceQAAsset : public UObject
+{
+	GENERATED_BODY()
+
+public:
+	void FindOrAddOperators(const TArray<FString>& InOperatorNames);
+
+	void RunAndAddTest(UNeuralNetwork* InOutNetwork, const FString& OperatorName);
+
+	bool CompareNewVsPreviousTests(const FString& InGroundTruthDirectory);
+
+	/**
+	 * It resets the new results so it can be run again.
+	 */
+	void FlushNewTests();
+
+	static UNeuralNetworkInferenceQAAsset* Load(const FString& InNeuralNetworkInferenceQAAssetParentDirectoryName, const FString& InNeuralNetworkInferenceQAAssetName);
+
+	bool Save();
+
+protected:
+	UPROPERTY(EditAnywhere, Category = "Neural Network Inference QA")
+	TMap<FString, FNeuralNetworkInferenceQAOperatorAsset> Operators;
+};
diff --git a/Engine/Plugins/Experimental/NeuralNetworkInference/Source/Private/NeuralNetworkInferenceQA/NeuralNetworkInferenceQA.Build.cs b/Engine/Plugins/Experimental/NeuralNetworkInference/Source/Private/NeuralNetworkInferenceQA/NeuralNetworkInferenceQA.Build.cs
new file mode 100644
--- /dev/null
+++ b/Engine/Plugins/Experimental/NeuralNetworkInference/Source/Private/NeuralNetworkInferenceQA/NeuralNetworkInferenceQA.Build.cs
@@ -0,0 +1,37 @@
+// Copyright Epic Games, Inc. All Rights Reserved.
+
+using UnrealBuildTool;
+
+public class NeuralNetworkInferenceQA : ModuleRules
+{
+	public NeuralNetworkInferenceQA( ReadOnlyTargetRules Target ) : base( Target )
+	{
+        ShortName = "NNIQA"; // Shorten to avoid path-too-long errors
+		PCHUsage = ModuleRules.PCHUsageMode.UseExplicitOrSharedPCHs;
+
+		PublicIncludePaths.AddRange(
+			new string[] {
+				System.IO.Path.Combine(ModuleDirectory, "..")
+			}
+		);
+
+		PublicDependencyModuleNames.AddRange
+			(
+			new string[] {
+				"Core",
+				"CoreUObject",
+				"Engine",
+				"NeuralNetworkInference"
+			}
+		);
+
+		PrivateDependencyModuleNames.AddRange
+			(
+			new string[] {
+				"ModelProto",
+				"ModelProtoFileReader",
+				"NeuralNetworkInferenceProfiling"
+			}
+		);
+	}
+}
diff --git a/Engine/Plugins/Experimental/NeuralNetworkInference/Source/Private/NeuralNetworkInferenceQA/Private/ModelProtoFileReaderTester.cpp b/Engine/Plugins/Experimental/NeuralNetworkInference/Source/Private/NeuralNetworkInferenceQA/Private/ModelProtoFileReaderTester.cpp
new file mode 100644
--- /dev/null
+++ b/Engine/Plugins/Experimental/NeuralNetworkInference/Source/Private/NeuralNetworkInferenceQA/Private/ModelProtoFileReaderTester.cpp
@@ -0,0 +1,26 @@
+// Copyright Epic Games, Inc. All Rights Reserved.
+
+#include "ModelProtoFileReaderTester.h"
+#include "ModelProtoFileReader.h"
+#include "NeuralNetworkInferenceQAUtils.h"
+
+
+
+/* FModelProtoFileReaderTester public functions
+ *****************************************************************************/
+
+bool FModelProtoFileReaderTester::Test(const FString& InONNXModelFileName)
+{
+	UE_LOG(LogNeuralNetworkInferenceQA, Display, TEXT("-------------------- Read ONNX Network And Test"));
+
+#if WITH_EDITOR
+	FModelProto ModelProto;
+	FModelProtoFileReader::ReadModelProtoFromFile(ModelProto, InONNXModelFileName);
+	UE_LOG(LogNeuralNetworkInferenceQA, Display, TEXT("ModelProto:\n%s"), *ModelProto.ToString());
+	ensureMsgf(ModelProto.IsLoaded(), TEXT("FModelProtoFileReaderTester::ONNXReadNetworkTest() failed, FModelProto could not be read from InONNXModelFileName: %s."), *InONNXModelFileName);
+	return ModelProto.IsLoaded();
+#else //WITH_EDITOR
+	UE_LOG(LogNeuralNetworkInferenceQA, Display, TEXT("ONNXReadNetworkTest test skipped (only if WITH_EDITOR enabled)."));
+	return true;
+#endif //WITH_EDITOR
+}
diff --git a/Engine/Plugins/Experimental/NeuralNetworkInference/Source/Private/NeuralNetworkInferenceQA/Private/ModelProtoFileReaderTester.h b/Engine/Plugins/Experimental/NeuralNetworkInference/Source/Private/NeuralNetworkInferenceQA/Private/ModelProtoFileReaderTester.h
new file mode 100644
--- /dev/null
+++ b/Engine/Plugins/Experimental/NeuralNetworkInference/Source/Private/NeuralNetworkInferenceQA/Private/ModelProtoFileReaderTester.h
@@ -0,0 +1,11 @@
+// Copyright Epic Games, Inc. All Rights Reserved.
+
+#pragma once
+
+#include "CoreMinimal.h"
+
+class FModelProtoFileReaderTester
+{
+public:
+	static bool Test(const FString& InONNXModelFileName);
+};
diff --git a/Engine/Plugins/Experimental/NeuralNetworkInference/Source/Private/NeuralNetworkInferenceQA/Private/ModelUnitTester.cpp b/Engine/Plugins/Experimental/NeuralNetworkInference/Source/Private/NeuralNetworkInferenceQA/Private/ModelUnitTester.cpp
new file mode 100644
--- /dev/null
+++ b/Engine/Plugins/Experimental/NeuralNetworkInference/Source/Private/NeuralNetworkInferenceQA/Private/ModelUnitTester.cpp
@@ -0,0 +1,582 @@
+// Copyright Epic Games, Inc. All Rights Reserved.
+
+#include "ModelUnitTester.h"
+#include "NeuralTimer.h"
+#include "NeuralNetworkInferenceQAUtils.h"
+#include "Misc/Paths.h"
+
+
+
+/* FModelUnitTester axuiliary functions
+ *****************************************************************************/
+
+struct FNNIUnitTesterTimeData
+{
+	FNeuralStatsData ComputeTimeData;
+	FNeuralStatsData InputCopyTimeData;
+	FNeuralStatsData OutputCopyTimeData;
+
+	FNNIUnitTesterTimeData(const FNeuralStatsData& InComputeTimeData, const FNeuralStatsData& InInputCopyTimeData, const FNeuralStatsData& InOutputCopyTimeData)
+		: ComputeTimeData(InComputeTimeData)
+		, InputCopyTimeData(InInputCopyTimeData)
+		, OutputCopyTimeData(InOutputCopyTimeData)
+	{}
+};
+
+FNNIUnitTesterTimeData FModelUnitTester_GetTimeInformation(UNeuralNetwork* InOutNetwork, TArray<float>& OutCPUGPUCPUOutput, const TArray<float>& InInputArray, const int32 InRepetitions)
+{
+	/* Input/output copy speed */
+	FNeuralStats OutCopyingStats;
+	FNeuralTimer Timer;
+	for (int32 TimerIndex = 0; TimerIndex < InRepetitions; ++TimerIndex)
+	{
+		InOutNetwork->SetInputFromArrayCopy(InInputArray);
+		Timer.Tic();
+		OutCPUGPUCPUOutput = InOutNetwork->GetOutputTensor().GetArrayCopy<float>();
+		const float CurrentCopyingTime = Timer.Toc();
+		OutCopyingStats.StoreSample(CurrentCopyingTime);
+
+	}
+
+	/* Forward() speed */
+	if (InRepetitions > 1)
+	{
+		for (int32 TimerIndex = 0; TimerIndex < 5; ++TimerIndex)
+		{
+			InOutNetwork->Run();
+		}
+	}
+
+	if (InRepetitions > 0)
+	{
+		for (int32 TimerIndex = 0; TimerIndex < InRepetitions; ++TimerIndex)
+		{
+			InOutNetwork->SetInputFromArrayCopy(InInputArray);
+			InOutNetwork->Run();
+			OutCPUGPUCPUOutput = InOutNetwork->GetOutputTensor().GetArrayCopy<float>();
+		}
+	}
+	// Return NetworkTimeData
+	return FNNIUnitTesterTimeData(InOutNetwork->GetInferenceStats(), InOutNetwork->GetInputMemoryTransferStats(), OutCopyingStats.GetStats());
+}
+
+
+
+/* FModelUnitTester static public functions
+ *****************************************************************************/
+
+bool FModelUnitTester::GlobalTest(const FString& InProjectContentDir, const FString& InModelZooRelativeDirectory)
+{
+	// NOTE: models are separated into multiple lines to enable/disable particular network for faster testing
+
+	// Model names, input values, and number of repetitions for profiling
+	const TArray<FString> ModelNames({
+		TEXT("MLRigDeformer"),
+		TEXT("cloth_network"),
+		TEXT("HS"),
+		TEXT("RL")
+	});
+
+	// Ground truths
+	const TArray<TArray<double>> CPUGroundTruths({
+		{3.728547, 0.008774, 4.595651, 212.193216, 742.434561, 4.250668, 4.717748},
+		{0.042571, 0.023693, 0.015783, 13.100505, 8.050994, 0.028807, 0.016387},
+		{138.372906, 126.753839, 127.287254, 130.316062, 127.303424, 124.800896, 126.546051},
+		{0.488662, 0.472437, 0.478862, 0.522685, 0.038322, 0.480848, 0.483821}
+	});
+	
+	const TArray<TArray<double>> GPUGroundTruths({
+		{3.728547, 0.008774, 4.595651, 212.193208, 742.434578, 4.250668, 4.717748},
+		{0.042571, 0.023693, 0.015783, 13.100504, 8.050994, 0.028807, 0.016387},
+		{138.373184, 126.754100, 127.287398, 130.316194, 127.303495, 124.801134, 126.5462530},
+		{0.488662, 0.472437, 0.478862, 0.522685, 0.038322, 0.480848, 0.483821}
+	});
+
+	const TArray<float> InputArrayValues({ // This one can be shorter than CPU/GPUGroundTruths
+		1.f,
+		0.f,
+		-1.f,
+		100.f,
+		-100.f,
+		0.5f,
+		-0.5f
+	});
+	
+	// Speed profiling test - 0 repetitions means that test will not be run
+#ifdef WITH_UE_AND_ORT_SUPPORT
+	const TArray<int32> CPURepetitionsForUEAndORTBackEnd({ 1000, 1000,  50, 1000 });
+#ifdef PLATFORM_WIN64
+	const TArray<int32> GPURepetitionsForUEAndORTBackEnd({ 1000, 1000, 100, 1000 });
+#else //PLATFORM_WIN64
+	UE_LOG(LogNeuralNetworkInferenceQA, Display, TEXT("FModelUnitTester::GlobalTest(): GPU tests disabled for non-Windows platforms."));
+	const TArray<int32> GPURepetitionsForUEAndORTBackEnd({ 0, 0, 0, 0 });
+#endif //PLATFORM_WIN64
+#else //WITH_UE_AND_ORT_SUPPORT
+	const TArray<int32> CPURepetitionsForUEAndORTBackEnd({ 0, 0, 0, 0 });
+	const TArray<int32> GPURepetitionsForUEAndORTBackEnd({ 0, 0, 0, 0 });
+#endif //WITH_UE_AND_ORT_SUPPORT
+#if WITH_EDITOR
+	const TArray<int32> CPURepetitionsForUEOnlyBackEnd({ 10, 0, 0, 0 });
+	const TArray<int32> GPURepetitionsForUEOnlyBackEnd({ 10, 10, 0, 0 });
+#else //WITH_EDITOR
+	const TArray<int32> CPURepetitionsForUEOnlyBackEnd({ 0, 0, 0, 0 });
+	const TArray<int32> GPURepetitionsForUEOnlyBackEnd({ 0, 0, 0, 0 });
+#endif //WITH_EDITOR
+	// Run tests
+	UE_LOG(LogNeuralNetworkInferenceQA, Display, TEXT("ENeuralBackEnd::UEOnly test for FModelUnitTester::GlobalTest disabled due to not being compatible with ORT (linking error issues)."));
+	return ModelLoadAccuracyAndSpeedTests(InProjectContentDir, InModelZooRelativeDirectory, ModelNames, InputArrayValues, CPUGroundTruths, GPUGroundTruths, CPURepetitionsForUEAndORTBackEnd, GPURepetitionsForUEAndORTBackEnd, CPURepetitionsForUEOnlyBackEnd, GPURepetitionsForUEOnlyBackEnd);
+}
+
+
+
+/* FModelUnitTester static private functions
+ *****************************************************************************/
+
+bool FModelUnitTester::ModelLoadAccuracyAndSpeedTests(const FString& InProjectContentDir, const FString& InModelZooRelativeDirectory, const TArray<FString>& InModelNames,
+	const TArray<float>& InInputArrayValues, const TArray<TArray<double>>& InCPUGroundTruths, const TArray<TArray<double>>& InGPUGroundTruths, const TArray<int32>& InCPURepetitionsForUEAndORTBackEnd,
+	const TArray<int32>& InGPURepetitionsForUEAndORTBackEnd, const TArray<int32>& InCPURepetitionsForUEOnlyBackEnd, const TArray<int32>& InGPURepetitionsForUEOnlyBackEnd)
+{
+	bool bDidGlobalTestPassed = true;
+
+	const FString ModelZooDirectory = InProjectContentDir / InModelZooRelativeDirectory;
+
+	// Test OTXT/UAsset accuracy
+	for (int32 ModelIndex = 0; ModelIndex < InModelNames.Num(); ++ModelIndex)
+	{
+		const FString& ModelName = InModelNames[ModelIndex];
+		const TArray<double>& CPUGroundTruths = InCPUGroundTruths[ModelIndex];
+		const TArray<double>& GPUGroundTruths = InGPUGroundTruths[ModelIndex];
+
+		UE_LOG(LogNeuralNetworkInferenceQA, Display, TEXT("--------------- %s - Network Uasset Load and Run"), *ModelName);
+		const FString UAssetModelFilePath = GetUAssetModelFilePath(ModelName, InModelZooRelativeDirectory);
+		UNeuralNetwork* Network = NetworkUassetLoadTest(UAssetModelFilePath);
+		if (!Network)
+		{
+			UE_LOG(LogNeuralNetworkInferenceQA, Warning, TEXT("UNeuralNetwork could not be loaded from UAssetModelFilePath %s."), *UAssetModelFilePath);
+			return false;
+		}
+		// Input debugging
+		UE_LOG(LogNeuralNetworkInferenceQA, Display, TEXT("--------------- %s - Input/Output"), *ModelName);
+		for (int32 TensorIndex = 0; TensorIndex < Network->GetInputTensorNumber(); ++TensorIndex)
+		{
+			UE_LOG(LogNeuralNetworkInferenceQA, Display, TEXT("InputTensor[%d] = %s."), TensorIndex, *Network->GetInputTensor(TensorIndex).GetName());
+		}
+		// Output debugging
+		for (int32 TensorIndex = 0; TensorIndex < Network->GetOutputTensorNumber(); ++TensorIndex)
+		{
+			UE_LOG(LogNeuralNetworkInferenceQA, Display, TEXT("OutputTensor[%d] = %s."), TensorIndex, *Network->GetOutputTensor(TensorIndex).GetName());
+		}
+		const bool bShouldRunUEAndORTBackEnd = (InCPURepetitionsForUEAndORTBackEnd[ModelIndex] + InGPURepetitionsForUEAndORTBackEnd[ModelIndex] > 0);
+		if (bShouldRunUEAndORTBackEnd)
+		{
+			bDidGlobalTestPassed &= ModelAccuracyTest(Network, ENeuralSynchronousMode::Synchronous, UNeuralNetwork::ENeuralBackEnd::UEAndORT, InInputArrayValues, CPUGroundTruths, GPUGroundTruths);
+			bDidGlobalTestPassed &= ModelAccuracyTest(Network, ENeuralSynchronousMode::Asynchronous, UNeuralNetwork::ENeuralBackEnd::UEAndORT, InInputArrayValues, CPUGroundTruths, GPUGroundTruths);
+		}
+		const bool bShouldRunUEOnlyBackEnd = false; // (InCPURepetitionsForUEOnlyBackEnd[ModelIndex] * InCPURepetitionsForUEOnlyBackEnd[ModelIndex] > 0);
+		if (bShouldRunUEOnlyBackEnd)
+		{
+			bDidGlobalTestPassed &= ModelAccuracyTest(Network, ENeuralSynchronousMode::Synchronous, UNeuralNetwork::ENeuralBackEnd::UEOnly, InInputArrayValues, CPUGroundTruths, GPUGroundTruths);
+		}
+
+		UE_LOG(LogNeuralNetworkInferenceQA, Display, TEXT("---------------------------------------------------------------------------------------------------------------------------------"));
+
+		UE_LOG(LogNeuralNetworkInferenceQA, Display, TEXT("--------------- %s - Network ONNX/ORT Load and Run"), *ModelName);
+#if WITH_EDITOR
+		const int32 ONNXOrORTIndex = 0; //for (int32 ONNXOrORTIndex = 0; ONNXOrORTIndex < 2; ++ONNXOrORTIndex)
+		{
+			const FString ModelFilePath = (ONNXOrORTIndex % 2 == 0 ? GetONNXModelFilePath(ModelZooDirectory, ModelName) : GetORTModelFilePath(ModelZooDirectory, ModelName));
+			const FString ModelType = (ONNXOrORTIndex % 2 == 0 ? TEXT("ONNX") : TEXT("ORT"));
+			UE_LOG(LogNeuralNetworkInferenceQA, Display, TEXT("-------------------- %s - Network %s Load and Run - %s"), *ModelName, *ModelType, *ModelFilePath);
+			if (bShouldRunUEAndORTBackEnd)
+			{
+				bDidGlobalTestPassed &= ModelAccuracyTest(NetworkONNXOrORTLoadTest(ModelFilePath), ENeuralSynchronousMode::Synchronous, UNeuralNetwork::ENeuralBackEnd::UEAndORT, InInputArrayValues, CPUGroundTruths, GPUGroundTruths);
+				bDidGlobalTestPassed &= ModelAccuracyTest(NetworkONNXOrORTLoadTest(ModelFilePath), ENeuralSynchronousMode::Asynchronous, UNeuralNetwork::ENeuralBackEnd::UEAndORT, InInputArrayValues, CPUGroundTruths, GPUGroundTruths);
+			}
+			if (bShouldRunUEOnlyBackEnd)
+			{
+				bDidGlobalTestPassed &= ModelAccuracyTest(NetworkONNXOrORTLoadTest(ModelFilePath), ENeuralSynchronousMode::Synchronous, UNeuralNetwork::ENeuralBackEnd::UEOnly, InInputArrayValues, CPUGroundTruths, GPUGroundTruths);
+			}
+		}
+#else //WITH_EDITOR
+	UE_LOG(LogNeuralNetworkInferenceQA, Display, TEXT("-------------------- Skipped (only if WITH_EDITOR enabled)."));
+#endif //WITH_EDITOR
+	}
+
+	// Profile speed
+	for (int32 ModelIndex = 0; ModelIndex < InModelNames.Num(); ++ModelIndex)
+	{
+		const FString& ModelName = InModelNames[ModelIndex];
+		UE_LOG(LogNeuralNetworkInferenceQA, Display, TEXT("--------------- %s - Network UAsset Speed Profiling"), *ModelName);
+		const FString UAssetModelFilePath = GetUAssetModelFilePath(ModelName, InModelZooRelativeDirectory);
+
+		// UEAndORT (if WITH_UE_AND_ORT_SUPPORT)
+#ifdef WITH_UE_AND_ORT_SUPPORT
+		bDidGlobalTestPassed &= ModelSpeedTest(UAssetModelFilePath, ENeuralDeviceType::CPU, UNeuralNetwork::ENeuralBackEnd::UEAndORT, InCPURepetitionsForUEAndORTBackEnd[ModelIndex]);
+		UNeuralNetwork* Network = NetworkUassetLoadTest(UAssetModelFilePath);
+		if (Network->GetBackEndForCurrentPlatform() != UNeuralNetwork::ENeuralBackEnd::UEAndORT)
+		{
+			UE_LOG(LogNeuralNetworkInferenceQA, Warning, TEXT("-------------------- Default UAsset BackEnd should be UEAndORT."));
+			return false;
+		}
+		else if (InGPURepetitionsForUEAndORTBackEnd[ModelIndex] > 0 && Network->IsGPUSupported())
+		{
+			bDidGlobalTestPassed &= ModelSpeedTest(UAssetModelFilePath, ENeuralDeviceType::GPU, UNeuralNetwork::ENeuralBackEnd::UEAndORT, InGPURepetitionsForUEAndORTBackEnd[ModelIndex]);
+		}
+		else
+		{
+			UE_LOG(LogNeuralNetworkInferenceQA, Display, TEXT("-------------------- ModelSpeedTest-UEAndORT-GPU skipped (DX12 not enabled or InGPURepetitionsForUEAndORTBackEnd[ModelIndex] = %d is 0)."),
+				InGPURepetitionsForUEAndORTBackEnd[ModelIndex]);
+		}
+#endif //WITH_UE_AND_ORT_SUPPORT
+
+		// UEOnly
+		const bool bShouldRunUEOnlyBackEnd = false;
+		if (bShouldRunUEOnlyBackEnd)
+		{
+			bDidGlobalTestPassed &= ModelSpeedTest(UAssetModelFilePath, ENeuralDeviceType::CPU, UNeuralNetwork::ENeuralBackEnd::UEOnly, InCPURepetitionsForUEOnlyBackEnd[ModelIndex]);
+			bDidGlobalTestPassed &= ModelSpeedTest(UAssetModelFilePath, ENeuralDeviceType::GPU, UNeuralNetwork::ENeuralBackEnd::UEOnly, InGPURepetitionsForUEOnlyBackEnd[ModelIndex]);
+		}
+	}
+
+	return bDidGlobalTestPassed;
+}
+
+FString FModelUnitTester::GetONNXModelFilePath(const FString& ModelZooDirectory, const FString& InModelName)
+{
+	return /*FPaths::ConvertRelativePathToFull*/(ModelZooDirectory / InModelName + TEXT("/") + InModelName + TEXT(".onnx")); // E.g., ModelZooDirectory / TEXT("ExampleNetworkReadable/ExampleNetworkReadable.onnx")
+}
+
+FString FModelUnitTester::GetORTModelFilePath(const FString& ModelZooDirectory, const FString& InModelName)
+{
+	return /*FPaths::ConvertRelativePathToFull*/(ModelZooDirectory / InModelName + TEXT("/") + InModelName + TEXT(".ort")); // E.g., ModelZooDirectory / TEXT("ExampleNetworkReadable/ExampleNetworkReadable.ort")
+}
+
+FString FModelUnitTester::GetUAssetModelFilePath(const FString& InModelName, const FString& InModelZooRelativeDirectory)
+{
+	return InModelName + TEXT("'/Game/") + InModelZooRelativeDirectory / InModelName + TEXT("/") + InModelName + TEXT(".") + InModelName + TEXT("'"); // E.g., '/Game/[MODEL_ZOO_DIR]/ExampleNetworkReadable/ExampleNetworkReadable.ExampleNetworkReadable'
+}
+
+UNeuralNetwork* FModelUnitTester::NetworkUassetLoadTest(const FString& InUAssetPath)
+{
+	UNeuralNetwork* Network = LoadObject<UNeuralNetwork>((UObject*)GetTransientPackage(), *InUAssetPath);
+	if (!Network)
+	{
+		ensureMsgf(false, TEXT("UNeuralNetwork is a nullptr. Path: \"%s\"."), *InUAssetPath);
+		return nullptr;
+	}
+	if (!Network->IsLoaded())
+	{
+		ensureMsgf(false, TEXT("UNeuralNetwork could not be loaded from uasset disk location. Path: \"%s\"."), *InUAssetPath);
+		return nullptr;
+	}
+	return Network;
+}
+
+UNeuralNetwork* FModelUnitTester::NetworkONNXOrORTLoadTest(const FString& InModelFilePath)
+{
+	// Load network architecture and weights from file
+	UNeuralNetwork* Network = NewObject<UNeuralNetwork>((UObject*)GetTransientPackage(), UNeuralNetwork::StaticClass());
+	if (!Network)
+	{
+		ensureMsgf(false, TEXT("UNeuralNetwork is a nullptr. Path: \"%s\"."), *InModelFilePath);
+		return nullptr;
+	}
+	if (!Network->Load(InModelFilePath))
+	{
+		ensureMsgf(false, TEXT("UNeuralNetwork could not be loaded from ONNX file disk location. Path: \"%s\"."), *InModelFilePath);
+		return nullptr;
+	}
+	return Network;
+}
+
+bool FModelUnitTester::ModelAccuracyTest(UNeuralNetwork* InOutNetwork, const ENeuralSynchronousMode InSynchronousMode, const UNeuralNetwork::ENeuralBackEnd InBackEnd, const TArray<float>& InInputArrayValues,
+	const TArray<double>& InCPUGroundTruths, const TArray<double>& InGPUGroundTruths)
+{
+	// Sanity check
+	if (!InOutNetwork)
+	{
+		return false;
+	}
+	// Find NetworkSize
+	const int64 NetworkSize = InOutNetwork->GetInputTensor().Num();
+	// Initialize input data
+	TArray<TArray<float>> InputArrays;
+	{
+		InputArrays.Reserve(InInputArrayValues.Num());
+		for (const float InputArrayValue : InInputArrayValues)
+		{
+			InputArrays.Emplace(TArray<float>({}));
+			InputArrays.Last().Init(InputArrayValue, NetworkSize);
+		}
+		if (InputArrays.Num() > InCPUGroundTruths.Num() || InputArrays.Num() > InGPUGroundTruths.Num())
+		{
+			UE_LOG(LogNeuralNetworkInferenceQA, Warning, TEXT("InputArrays.Num() <= InCPUGroundTruths.Num() && InputArrays.Num() <= InGPUGroundTruths.Num() failed: %d vs. %d vs. %d."),
+				InputArrays.Num(), InCPUGroundTruths.Num(), InGPUGroundTruths.Num());
+			return false;
+		}
+	}
+	
+	// Save original network state
+	const ENeuralDeviceType OriginalDeviceType = InOutNetwork->GetDeviceType();
+	const ENeuralDeviceType OriginalInputDeviceType = InOutNetwork->GetInputDeviceType();
+	const ENeuralDeviceType OriginalOutputDeviceType = InOutNetwork->GetOutputDeviceType();
+	const ENeuralSynchronousMode OriginalSynchronousMode = InOutNetwork->GetSynchronousMode();
+	const ENeuralThreadMode OriginalDelegateThreadMode = InOutNetwork->GetOnAsyncRunCompletedDelegateMode();
+	const UNeuralNetwork::ENeuralBackEnd OriginalBackEnd = InOutNetwork->GetBackEnd();
+	
+	// Set (a)synchronous Mode
+	InOutNetwork->SetSynchronousMode(InSynchronousMode);
+	InOutNetwork->SetOnAsyncRunCompletedDelegateMode(ENeuralThreadMode::AnyThread);
+
+	// Set back end
+	if (!InOutNetwork->SetBackEnd(InBackEnd))
+	{
+		UE_LOG(LogNeuralNetworkInferenceQA, Warning, TEXT("Backend %s is disabled."), *GetBackEndString(InBackEnd));
+		return false;
+	}
+	const FString BackEndString = GetBackEndString(InBackEnd);
+	
+	// Run each input with CPU/GPU and compare with each other and with the ground truth
+	// Multiple for loops to make sure that running on the CPU does not affect GPU results or vice versa
+	bool bDidGlobalTestPassed = true;
+	TArray<TArray<float>> CPUOutputs, CPUGPUCPUOutputs, CPUGPUGPUOutputs, GPUGPUCPUOutputs;
+	
+	// Input CPU + Network CPU + Output CPU
+	ModelAccuracyTestRun(CPUOutputs, InOutNetwork, InputArrays, /*DeviceType*/ENeuralDeviceType::CPU, /*InputDeviceType*/ENeuralDeviceType::CPU, /*OutputDeviceType*/ENeuralDeviceType::CPU);
+
+	if (InOutNetwork->IsGPUSupported())
+	{
+		// Input CPU + Network GPU + Output CPU
+		ModelAccuracyTestRun(CPUGPUCPUOutputs, InOutNetwork, InputArrays, /*DeviceType*/ENeuralDeviceType::GPU, /*InputDeviceType*/ENeuralDeviceType::CPU, /*OutputDeviceType*/ENeuralDeviceType::CPU);
+		// Input CPU + Network GPU + Output GPU
+		ModelAccuracyTestRun(CPUGPUGPUOutputs, InOutNetwork, InputArrays, /*DeviceType*/ENeuralDeviceType::GPU, /*InputDeviceType*/ENeuralDeviceType::CPU, /*OutputDeviceType*/ENeuralDeviceType::GPU);
+		// Input GPU + Network GPU + Output CPU
+		ModelAccuracyTestRun(GPUGPUCPUOutputs, InOutNetwork, InputArrays, /*DeviceType*/ENeuralDeviceType::GPU, /*InputDeviceType*/ENeuralDeviceType::GPU, /*OutputDeviceType*/ENeuralDeviceType::CPU);
+	}
+	else
+	{
+		UE_LOG(LogNeuralNetworkInferenceQA, Display, TEXT("-------------------- ModelAccuracyTest-UEAndORT-GPU skipped (DX12 not enabled)."));
+	}
+
+	for (int32 Index = 0; Index < InputArrays.Num(); ++Index)
+	{
+		const TArray<float>& GPUGPUCPUOutput = (InOutNetwork->IsGPUSupported() ? GPUGPUCPUOutputs[Index] : CPUOutputs[Index]);
+		const TArray<float>& CPUGPUCPUOutput = (InOutNetwork->IsGPUSupported() ? CPUGPUCPUOutputs[Index] : CPUOutputs[Index]);
+		const TArray<float>& CPUGPUGPUOutput = (InOutNetwork->IsGPUSupported() ? CPUGPUGPUOutputs[Index] : CPUOutputs[Index]);
+		// Prepare verbose
+		const double CPUAvgL1Norm = GetAveragedL1Norm(CPUOutputs[Index]);
+		const double CPUGPUCPUAvgL1Norm = GetAveragedL1Norm(CPUGPUCPUOutput);
+		const double GPUGPUCPUAvgL1Norm = GetAveragedL1Norm(GPUGPUCPUOutput);
+		const double CPUGPUGPUAvgL1Norm = GetAveragedL1Norm(CPUGPUGPUOutput);
+		const double RelativeCoefficient = 1. / FMath::Max(1., FMath::Min(CPUAvgL1Norm, CPUGPUCPUAvgL1Norm)); // Max(1, X) to avoid 0s
+		const double CPUGPUAvgL1NormDiff = GetAveragedL1NormDiff(CPUOutputs[Index], CPUGPUCPUOutput) * RelativeCoefficient * 1e6;
+		const double GPUGPUInputAvgL1NormDiff = GetAveragedL1NormDiff(CPUGPUCPUOutput, GPUGPUCPUOutput) / FMath::Max(1., FMath::Min(GPUGPUCPUAvgL1Norm, CPUGPUCPUAvgL1Norm)) * 1e12;
+		const double GPUGPUOutputAvgL1NormDiff = GetAveragedL1NormDiff(CPUGPUCPUOutput, CPUGPUGPUOutput) / FMath::Max(1., FMath::Min(CPUGPUGPUAvgL1Norm, CPUGPUCPUAvgL1Norm)) * 1e12;
+		const double FastCPUGPUAvgL1NormDiff = FMath::Abs((CPUAvgL1Norm - CPUGPUCPUAvgL1Norm)) * RelativeCoefficient * 1e6;
+		const double FastCPUAvgL1NormDiff = FMath::Abs(CPUAvgL1Norm - InCPUGroundTruths[Index]) / FMath::Max(1., FMath::Min(CPUAvgL1Norm, InCPUGroundTruths[Index])) * 1e7;
+		const double FastGPUAvgL1NormDiff = FMath::Abs(CPUGPUCPUAvgL1Norm - InGPUGroundTruths[Index]) / FMath::Max(1., FMath::Min(CPUGPUCPUAvgL1Norm, InGPUGroundTruths[Index])) * 1e7;
+		// Print verbose
+		UE_LOG(LogNeuralNetworkInferenceQA, Display, TEXT("%s: InputNorm = %f, OutputNormCPU = %f, OutputNormGPU = %f, OutputNormCPUGPUGPU = %f, OutputNormGT = %f, CPUAvgL1Norm = %f, CPUGPUCPUAvgL1Norm = %f,"),
+			*BackEndString, GetAveragedL1Norm(InputArrays[Index]), CPUAvgL1Norm, CPUGPUCPUAvgL1Norm, CPUGPUGPUAvgL1Norm, InCPUGroundTruths[Index], CPUAvgL1Norm, CPUGPUCPUAvgL1Norm);
+		UE_LOG(LogNeuralNetworkInferenceQA, Display, TEXT("\tCPUGPUAvgL1NormDiff = %fe-6, GPUGPUInputAvgL1NormDiff = %fe-7, GPUGPUOutputAvgL1NormDiff = %fe-7, FastCPUGPUAvgL1NormDiff = %fe-6, FastCPUAvgL1NormDiff = %fe-7, FastGPUAvgL1NormDiff = %fe-7 (1e-7 is roughly the precision for float)."),
+			CPUGPUAvgL1NormDiff, GPUGPUInputAvgL1NormDiff, GPUGPUOutputAvgL1NormDiff, FastCPUGPUAvgL1NormDiff, FastCPUAvgL1NormDiff, FastGPUAvgL1NormDiff);
+		// Check if test failed and (if so) display information
+		const bool bDidSomeTestFailed = (!FMath::IsFinite(FastCPUGPUAvgL1NormDiff) || FastCPUGPUAvgL1NormDiff > 5)
+			|| (!FMath::IsFinite(CPUGPUAvgL1NormDiff) || CPUGPUAvgL1NormDiff > 5)
+			|| (!FMath::IsFinite(GPUGPUInputAvgL1NormDiff) || GPUGPUInputAvgL1NormDiff > 1)
+			|| (!FMath::IsFinite(GPUGPUOutputAvgL1NormDiff) || GPUGPUOutputAvgL1NormDiff > 1)
+			|| (!FMath::IsFinite(FastCPUAvgL1NormDiff) || FastCPUAvgL1NormDiff > 30)
+			|| (!FMath::IsFinite(FastGPUAvgL1NormDiff) || FastGPUAvgL1NormDiff > 30);
+		if (bDidSomeTestFailed)
+		{
+			const int64 MaxNumberElementsToDisplay = 100;
+			const TArray<int64>& InputSizes = InOutNetwork->GetInputTensor().GetSizes();
+			const TArray<int64>& OutputSizes = InOutNetwork->GetOutputTensor().GetSizes();
+			UE_LOG(LogNeuralNetworkInferenceQA, Display, TEXT("FastCPUGPUAvgL1NormDiff (%fe-6) < 5e-6 might have failed."), FastCPUGPUAvgL1NormDiff);
+			UE_LOG(LogNeuralNetworkInferenceQA, Display, TEXT("CPUGPUAvgL1NormDiff (%fe-6) < 5e-6 might have failed."), CPUGPUAvgL1NormDiff);
+			UE_LOG(LogNeuralNetworkInferenceQA, Display, TEXT("GPUGPUInputAvgL1NormDiff (%fe-12) < 1e-12 might have failed."), GPUGPUInputAvgL1NormDiff);
+			UE_LOG(LogNeuralNetworkInferenceQA, Display, TEXT("GPUGPUOutputAvgL1NormDiff (%fe-12) < 1e-12 might have failed."), GPUGPUOutputAvgL1NormDiff);
+			UE_LOG(LogNeuralNetworkInferenceQA, Display, TEXT("FastCPUAvgL1NormDiff (%fe-7) < 30e-7 might have failed (~30 times the float precision).\nCPUOutput = %s."),
+				FastCPUAvgL1NormDiff, *FNeuralTensor(CPUOutputs[Index], OutputSizes).ToString(MaxNumberElementsToDisplay));
+			UE_LOG(LogNeuralNetworkInferenceQA, Display, TEXT("FastGPUAvgL1NormDiff (%fe-7) < 30e-7 might have failed (~30 times the float precision).\nCPUGPUCPUOutput = %s."),
+				FastGPUAvgL1NormDiff, *FNeuralTensor(CPUGPUCPUOutput, OutputSizes).ToString(MaxNumberElementsToDisplay));
+			UE_LOG(LogNeuralNetworkInferenceQA, Display, TEXT("Input = %s"),
+				*FNeuralTensor(InOutNetwork->GetInputTensor().GetArrayCopy<float>(), InputSizes).ToString(MaxNumberElementsToDisplay));
+			UE_LOG(LogNeuralNetworkInferenceQA, Display, TEXT("CPUOutput = %s"), *FNeuralTensor(CPUOutputs[Index], OutputSizes).ToString(MaxNumberElementsToDisplay));
+			UE_LOG(LogNeuralNetworkInferenceQA, Display, TEXT("CPUGPUCPUOutput = %s"), *FNeuralTensor(CPUGPUCPUOutput, OutputSizes).ToString(MaxNumberElementsToDisplay));
+			UE_LOG(LogNeuralNetworkInferenceQA, Display, TEXT("GPUGPUCPUOutput = %s"), *FNeuralTensor(GPUGPUCPUOutput, OutputSizes).ToString(MaxNumberElementsToDisplay));
+			UE_LOG(LogNeuralNetworkInferenceQA, Display, TEXT("CPUGPUGPUOutput = %s"), *FNeuralTensor(CPUGPUGPUOutput, OutputSizes).ToString(MaxNumberElementsToDisplay));
+			UE_LOG(LogNeuralNetworkInferenceQA, Warning, TEXT("At least 1 of the 6 CPU/GPU tests failed."));
+			return false;
+		}
+		bDidGlobalTestPassed &= !bDidSomeTestFailed;
+	}
+	// Reset to original network state
+	InOutNetwork->SetDeviceType(/*DeviceType*/OriginalDeviceType, /*InputDeviceType*/OriginalInputDeviceType, /*OutputDeviceType*/OriginalOutputDeviceType);
+	InOutNetwork->SetSynchronousMode(OriginalSynchronousMode);
+	InOutNetwork->SetOnAsyncRunCompletedDelegateMode(OriginalDelegateThreadMode);
+	InOutNetwork->SetBackEnd(OriginalBackEnd);
+	// Test successful
+	return bDidGlobalTestPassed;
+}
+
+void FModelUnitTester::ModelAccuracyTestRun(TArray<TArray<float>>& OutOutputs, UNeuralNetwork* InOutNetwork, const TArray<TArray<float>> InInputArrays,
+	const ENeuralDeviceType InDeviceType, const ENeuralDeviceType InInputDeviceType, const ENeuralDeviceType InOutputDeviceType)
+{
+	std::atomic<bool> bDidAsyncNetworkFinished(false);
+	const ENeuralSynchronousMode SynchronousMode = InOutNetwork->GetSynchronousMode();
+	if (SynchronousMode == ENeuralSynchronousMode::Asynchronous)
+	{
+		InOutNetwork->GetOnAsyncRunCompletedDelegate().BindLambda([&bDidAsyncNetworkFinished]()
+		{
+			bDidAsyncNetworkFinished = true;
+		});
+	}
+	for (int32 Index = 0; Index < InInputArrays.Num(); ++Index)
+	{
+		bDidAsyncNetworkFinished = false;
+		InOutNetwork->SetDeviceType(InDeviceType, InInputDeviceType, InOutputDeviceType);
+		InOutNetwork->SetInputFromArrayCopy(InInputArrays[Index]);
+		if (InInputDeviceType == ENeuralDeviceType::GPU)
+		{
+			InOutNetwork->InputTensorsToGPU();
+		}
+		InOutNetwork->Run();
+		// If async test, wait until inference is completed
+		if (SynchronousMode == ENeuralSynchronousMode::Asynchronous)
+		{
+			while (!bDidAsyncNetworkFinished)
+			{
+				FPlatformProcess::Sleep(0.1e-3);
+			}
+		}
+		if (InOutputDeviceType == ENeuralDeviceType::GPU)
+		{
+			InOutNetwork->OutputTensorsToCPU();
+		}
+		OutOutputs.Emplace(InOutNetwork->GetOutputTensor().GetArrayCopy<float>());
+	}
+	if (SynchronousMode == ENeuralSynchronousMode::Asynchronous)
+	{
+		InOutNetwork->GetOnAsyncRunCompletedDelegate().Unbind();
+	}
+}
+
+bool FModelUnitTester::ModelSpeedTest(const FString& InUAssetPath, const ENeuralDeviceType InDeviceType, const UNeuralNetwork::ENeuralBackEnd InBackEnd, const int32 InRepetitions)
+{
+	// Get debug strings
+	const FString DeviceTypeString = GetDeviceTypeString(InDeviceType);
+	const FString BackEndString = GetBackEndString(InBackEnd);
+	// Skip test if reps = 0
+	if (InRepetitions < 1)
+	{
+		UE_LOG(LogNeuralNetworkInferenceQA, Display, TEXT("ModelSpeedTest skipped for uasset \"%s\" for InDeviceType = %s (%d) and InBackEnd = %s (%d)."), *InUAssetPath, *DeviceTypeString, (int32)InDeviceType, *BackEndString, (int32)InBackEnd);
+		return true;
+	}
+	// Load Network
+	UNeuralNetwork* InOutNetwork = NetworkUassetLoadTest(InUAssetPath);
+	// Sanity checks
+	if (!InOutNetwork)
+	{
+		UE_LOG(LogNeuralNetworkInferenceQA, Warning, TEXT("FModelUnitTester::ModelSpeedTest(): InOutNetwork was a nullptr. Path: \"%s\"."), *InUAssetPath);
+		return false;
+	}
+	// Save original network state
+	const ENeuralDeviceType OriginalDeviceType = InOutNetwork->GetDeviceType();
+	const UNeuralNetwork::ENeuralBackEnd OriginalBackEnd = InOutNetwork->GetBackEnd();
+	// Set desired back end
+	if (!InOutNetwork->SetBackEnd(InBackEnd))
+	{
+		UE_LOG(LogNeuralNetworkInferenceQA, Warning, TEXT("Backend %s is disabled."), *GetBackEndString(InBackEnd));
+		return false;
+	}
+	// Needed variables
+	const int64 NetworkSize = InOutNetwork->GetInputTensor().Num();
+	TArray<float> InputArray;
+	InputArray.Init(1.f, NetworkSize);
+	TArray<float> CPUGPUCPUOutput;
+	// Run profiling 1 time
+	InOutNetwork->SetDeviceType(InDeviceType);
+	InOutNetwork->ResetStats();
+	const FNNIUnitTesterTimeData TimeData1 = FModelUnitTester_GetTimeInformation(InOutNetwork, CPUGPUCPUOutput, InputArray, 1);
+
+	// Run profiling n times
+	InOutNetwork->SetDeviceType(InDeviceType);
+	InOutNetwork->ResetStats();
+	const FNNIUnitTesterTimeData TimeDataN = FModelUnitTester_GetTimeInformation(InOutNetwork, CPUGPUCPUOutput, InputArray, InRepetitions);
+
+	// Display speed times
+	UE_LOG(LogNeuralNetworkInferenceQA, Display,
+		TEXT("%s-%s:, Times(msec)-> Inference = %f, Copy In = %f, Copy Out %f, avg(%d times) = Inference = %f, Copy In = %f, Copy Out %f."),
+		*BackEndString, *DeviceTypeString, TimeData1.ComputeTimeData.Average, TimeData1.InputCopyTimeData.Average, TimeData1.OutputCopyTimeData.Average,
+		InRepetitions, TimeDataN.ComputeTimeData.Average, TimeDataN.InputCopyTimeData.Average, TimeDataN.OutputCopyTimeData.Average);
+
+	// Reset to original network state
+	InOutNetwork->SetDeviceType(OriginalDeviceType);
+	InOutNetwork->SetBackEnd(OriginalBackEnd);
+	// Test successful
+	return true;
+}
+
+double FModelUnitTester::GetAveragedL1Norm(const TArray<float>& InArray)
+{
+	double AveragedL1Norm = 0.;
+	for (const float Value : InArray)
+	{
+		AveragedL1Norm += FMath::Abs(Value);
+	}
+	AveragedL1Norm /= InArray.Num();
+	return AveragedL1Norm;
+}
+
+double FModelUnitTester::GetAveragedL1NormDiff(const TArray<float>& InArray1, const TArray<float>& InArray2)
+{
+	// Sanity check
+	if (InArray1.Num() != InArray2.Num())
+	{
+		ensureMsgf(false, TEXT("InArray1.Num() == InArray2.Num() failed: %d != %d."), InArray1.Num(), InArray2.Num());
+		return -1.;
+	}
+	// Averaged L1 norm
+	double AveragedL1NormDiff = 0.;
+	for (int32 Index = 0; Index < InArray1.Num(); ++Index)
+	{
+		AveragedL1NormDiff += FMath::Abs(InArray1[Index] - InArray2[Index]);
+	}
+	AveragedL1NormDiff /= InArray1.Num();
+	return AveragedL1NormDiff;
+}
+
+FString FModelUnitTester::GetDeviceTypeString(const ENeuralDeviceType InDeviceType)
+{
+	if (InDeviceType == ENeuralDeviceType::CPU)
+	{
+		return TEXT("CPU");
+	}
+	else if (InDeviceType == ENeuralDeviceType::GPU)
+	{
+		return TEXT("GPU");
+	}
+	else
+	{
+		ensureMsgf(false, TEXT("Unknown DeviceType = %d."), (int32)InDeviceType);
+		return TEXT("");
+	}
+}
+
+FString FModelUnitTester::GetBackEndString(const UNeuralNetwork::ENeuralBackEnd InBackEnd)
+{
+	if (InBackEnd == UNeuralNetwork::ENeuralBackEnd::UEAndORT)
+	{
+		return TEXT("UEAndORT");
+	}
+	else if (InBackEnd == UNeuralNetwork::ENeuralBackEnd::UEOnly)
+	{
+		return TEXT("UEOnly");
+	}
+	else
+	{
+		ensureMsgf(false, TEXT("Unknown BackEndForCurrentPlatform = %d."), (int32)InBackEnd);
+		return TEXT("");
+	}
+}
diff --git a/Engine/Plugins/Experimental/NeuralNetworkInference/Source/Private/NeuralNetworkInferenceQA/Private/ModelUnitTester.h b/Engine/Plugins/Experimental/NeuralNetworkInference/Source/Private/NeuralNetworkInferenceQA/Private/ModelUnitTester.h
new file mode 100644
--- /dev/null
+++ b/Engine/Plugins/Experimental/NeuralNetworkInference/Source/Private/NeuralNetworkInferenceQA/Private/ModelUnitTester.h
@@ -0,0 +1,41 @@
+// Copyright Epic Games, Inc. All Rights Reserved.
+
+#pragma once
+
+#include "CoreMinimal.h"
+#include "NeuralNetwork.h"
+
+class FModelUnitTester
+{
+public:
+	static bool GlobalTest(const FString& InProjectContentDir, const FString& InModelZooRelativeDirectory);
+
+private:
+	/**
+	 * It runs a full model test on the desired model.
+	 */
+	static bool ModelLoadAccuracyAndSpeedTests(const FString& InProjectContentDir, const FString& InModelZooRelativeDirectory, const TArray<FString>& InModelNames,
+		const TArray<float>& InInputArrayValues, const TArray<TArray<double>>& InCPUGroundTruths, const TArray<TArray<double>>& InGPUGroundTruths,
+		const TArray<int32>& InCPURepetitionsForUEAndORT, const TArray<int32>& InGPURepetitionsForUEAndORTBackEnd, const TArray<int32>& InCPURepetitionsForUEOnly,
+		const TArray<int32>& InGPURepetitionsForUEOnly);
+	/**
+	 * Other auxiliary functions for GlobalTest().
+	 */
+	static FString GetONNXModelFilePath(const FString& InModelZooDirectory, const FString& InModelName);
+	static FString GetORTModelFilePath(const FString& InModelZooDirectory, const FString& InModelName);
+	static FString GetUAssetModelFilePath(const FString& InModelName, const FString& InModelZooRelativeDirectory);
+	static UNeuralNetwork* NetworkUassetLoadTest(const FString& InUAssetPath);
+	static UNeuralNetwork* NetworkONNXOrORTLoadTest(const FString& InModelFilePath);
+	static bool ModelAccuracyTest(UNeuralNetwork* InOutNetwork, const ENeuralSynchronousMode InSynchronousMode, const UNeuralNetwork::ENeuralBackEnd InBackEnd, const TArray<float>& InInputArrayValues,
+		const TArray<double>& InCPUGroundTruths, const TArray<double>& InGPUGroundTruths);
+	static void ModelAccuracyTestRun(TArray<TArray<float>>& OutOutputs, UNeuralNetwork* InOutNetwork, const TArray<TArray<float>> InInputArrays,
+		const ENeuralDeviceType InDeviceType, const ENeuralDeviceType InInputDeviceType, const ENeuralDeviceType InOutputDeviceType);
+	static bool ModelSpeedTest(const FString& InUAssetPath, const ENeuralDeviceType InDeviceType, const UNeuralNetwork::ENeuralBackEnd InBackEnd, const int32 InRepetitions);
+	static double GetAveragedL1Norm(const TArray<float>& InArray);
+	static double GetAveragedL1NormDiff(const TArray<float>& InArray1, const TArray<float>& InArray2);
+	/**
+	 * Verbose string auxiliary functions.
+	 */
+	static FString GetDeviceTypeString(const ENeuralDeviceType InDeviceType);
+	static FString GetBackEndString(const UNeuralNetwork::ENeuralBackEnd InBackEnd);
+};
diff --git a/Engine/Plugins/Experimental/NeuralNetworkInference/Source/Private/NeuralNetworkInferenceQA/Private/Module.cpp b/Engine/Plugins/Experimental/NeuralNetworkInference/Source/Private/NeuralNetworkInferenceQA/Private/Module.cpp
new file mode 100644
--- /dev/null
+++ b/Engine/Plugins/Experimental/NeuralNetworkInference/Source/Private/NeuralNetworkInferenceQA/Private/Module.cpp
@@ -0,0 +1,16 @@
+// Copyright Epic Games, Inc. All Rights Reserved.
+
+#include "Module.h"
+
+// This code will execute after your module is loaded into memory; the exact timing is specified in the .uplugin file per-module
+void FNeuralNetworkInferenceQAModule::StartupModule()
+{
+}
+
+// This function may be called during shutdown to clean up your module. For modules that support dynamic reloading,
+// we call this function before unloading the module.
+void FNeuralNetworkInferenceQAModule::ShutdownModule()
+{
+}
+
+IMPLEMENT_MODULE(FNeuralNetworkInferenceQAModule, NeuralNetworkInferenceQA);
diff --git a/Engine/Plugins/Experimental/NeuralNetworkInference/Source/Private/NeuralNetworkInferenceQA/Private/Module.h b/Engine/Plugins/Experimental/NeuralNetworkInference/Source/Private/NeuralNetworkInferenceQA/Private/Module.h
new file mode 100644
--- /dev/null
+++ b/Engine/Plugins/Experimental/NeuralNetworkInference/Source/Private/NeuralNetworkInferenceQA/Private/Module.h
@@ -0,0 +1,14 @@
+// Copyright Epic Games, Inc. All Rights Reserved.
+
+#pragma once
+
+#include "CoreMinimal.h"
+#include "Modules/ModuleManager.h"
+
+class FNeuralNetworkInferenceQAModule : public IModuleInterface
+{
+public:
+	// Begin IModuleInterface
+	virtual void StartupModule() override;
+	virtual void ShutdownModule() override;
+};
diff --git a/Engine/Plugins/Experimental/NeuralNetworkInference/Source/Private/NeuralNetworkInferenceQA/Private/NeuralNetworkInferenceQA.cpp b/Engine/Plugins/Experimental/NeuralNetworkInference/Source/Private/NeuralNetworkInferenceQA/Private/NeuralNetworkInferenceQA.cpp
new file mode 100644
--- /dev/null
+++ b/Engine/Plugins/Experimental/NeuralNetworkInference/Source/Private/NeuralNetworkInferenceQA/Private/NeuralNetworkInferenceQA.cpp
@@ -0,0 +1,49 @@
+// Copyright Epic Games, Inc. All Rights Reserved.
+
+#include "NeuralNetworkInferenceQA.h"
+#include "NeuralNetworkInferenceQAUtils.h"
+#include "UnitTester.h"
+#include "Misc/Paths.h"
+
+
+
+/* UNeuralNetworkInferenceQA static functions
+ *****************************************************************************/
+
+bool UNeuralNetworkInferenceQA::UnitTesting()
+{
+	const FString ProjectContentDir = FPaths::ProjectContentDir();
+	const FString MachineLearningTestsRelativeDirectory = TEXT("Tests/MachineLearning/");
+	const FString ModelZooRelativeDirectory = MachineLearningTestsRelativeDirectory / TEXT("Models/"); // Eg "Tests/MachineLearning/Models/"
+	const FString UnitTestRelativeDirectory = MachineLearningTestsRelativeDirectory / TEXT("UnitTests/"); // Eg "Tests/MachineLearning/UnitTests/"
+	const bool bDidGlobalTestPassed = FUnitTester::GlobalTest(ProjectContentDir, ModelZooRelativeDirectory, UnitTestRelativeDirectory);
+	ensureMsgf(bDidGlobalTestPassed, TEXT("UNeuralNetworkInferenceQA::UnitTesting() returned false. See log above."));
+	return bDidGlobalTestPassed;
+}
+
+
+
+#if WITH_DEV_AUTOMATION_TESTS
+
+#include "Misc/AutomationTest.h"
+
+IMPLEMENT_SIMPLE_AUTOMATION_TEST(FNeuralNetworkInferenceTest, "System.Engine.MachineLearning.NeuralNetworkInference (NNI)",
+	EAutomationTestFlags::ApplicationContextMask // = EditorContext | ClientContext | ServerContext | CommandletContext
+	| EAutomationTestFlags::FeatureMask // = NonNullRHI | RequiresUser
+	| EAutomationTestFlags::EngineFilter)
+
+bool FNeuralNetworkInferenceTest::RunTest(const FString& Parameters)
+{
+// #if WITH_EDITOR
+	if (!UNeuralNetworkInferenceQA::UnitTesting())
+	{
+		AddError(TEXT("FNeuralNetworkInferenceTest::RunTest(): UNeuralNetworkInferenceQA::UnitTesting() returned false. See log above.")); // Both AddWarning/AddError would work
+	}
+// #else //WITH_EDITOR
+// 	UE_LOG(LogNeuralNetworkInferenceQA, Warning, TEXT("@todo: FNeuralNetworkInferenceTest::RunTest() temporarily disabled on game due to JIRA"
+// 		" jira.it.epicgames.com/browse/UE-134683, UAssets not found on the Horde cooked tests."));
+// #endif //WITH_EDITOR
+	return true;
+}
+
+#endif //WITH_DEV_AUTOMATION_TESTS
diff --git a/Engine/Plugins/Experimental/NeuralNetworkInference/Source/Private/NeuralNetworkInferenceQA/Private/NeuralNetworkInferenceQAAsset.cpp b/Engine/Plugins/Experimental/NeuralNetworkInference/Source/Private/NeuralNetworkInferenceQA/Private/NeuralNetworkInferenceQAAsset.cpp
new file mode 100644
--- /dev/null
+++ b/Engine/Plugins/Experimental/NeuralNetworkInference/Source/Private/NeuralNetworkInferenceQA/Private/NeuralNetworkInferenceQAAsset.cpp
@@ -0,0 +1,260 @@
+// Copyright Epic Games, Inc. All Rights Reserved.
+
+#include "NeuralNetworkInferenceQAAsset.h"
+#include "NeuralNetworkInferenceQAUtils.h"
+#include "HAL/FileManager.h"
+#include "Misc/FileHelper.h"
+#include "Misc/Paths.h"
+#include "UObject/SavePackage.h"
+
+
+/* FNeuralNetworkInferenceQAOperatorTestAsset public functions
+ *****************************************************************************/
+
+void FNeuralNetworkInferenceQAOperatorTestAsset::AddInputTensors(const UNeuralNetwork* const InNetwork)
+{
+	for (int32 InputIndex = 0; InputIndex < InNetwork->GetInputTensorNumber(); ++InputIndex)
+	{
+		InputTensors.Push(InNetwork->GetInputTensor(InputIndex));
+	}
+}
+
+void FNeuralNetworkInferenceQAOperatorTestAsset::AddOutputTensors(const UNeuralNetwork* const InNetwork)
+{
+	for (int32 OutputIndex = 0; OutputIndex < InNetwork->GetOutputTensorNumber(); ++OutputIndex)
+	{
+		OutputTensors.Push(InNetwork->GetOutputTensor(OutputIndex));
+	}
+}
+
+bool FNeuralNetworkInferenceQAOperatorTestAsset::CompareAverageL1DiffNewVsPreviousTests(const FNeuralNetworkInferenceQAOperatorTestAsset& InOperatorTestAsset1,
+	const FNeuralNetworkInferenceQAOperatorTestAsset& InOperatorTestAsset2, const FString& InOperatorName)
+{
+	// Compare each input and output
+	return (CompareAverageL1DiffNewVsPreviousTests(InOperatorTestAsset1.InputTensors, InOperatorTestAsset2.InputTensors, InOperatorName)
+		&& CompareAverageL1DiffNewVsPreviousTests(InOperatorTestAsset1.OutputTensors, InOperatorTestAsset2.OutputTensors, InOperatorName));
+}
+
+bool FNeuralNetworkInferenceQAOperatorTestAsset::CompareAverageL1DiffNewVsPreviousTests(const TArray<FNeuralTensor>& InTensorsA, const TArray<FNeuralTensor>& InTensorsB, const FString& InOperatorName)
+{
+	// Number of inputs and outputs should be the same
+	if (InTensorsA.Num() != InTensorsB.Num())
+	{
+		return false;
+	}
+	// Compare each input
+	const FString DebugName = InOperatorName + TEXT("-CPUvsGT");
+	for (int32 TensorIndex = 0; TensorIndex < InTensorsA.Num(); ++TensorIndex)
+	{
+		if (!FNeuralNetworkInferenceQAUtils::EstimateTensorL1DiffError(InTensorsA[TensorIndex], InTensorsB[TensorIndex], /*ZeroThreshold*/5e-4, DebugName))
+		{
+			return false;
+		}
+	}
+	return true;
+}
+
+
+
+/* FNeuralNetworkInferenceQAOperatorAsset public functions
+ *****************************************************************************/
+
+FString FNeuralNetworkInferenceQAOperatorAsset::RunNetworkCPUAndGetString(UNeuralNetwork* InOutNetwork)
+{
+	// Create input string and InputTensorArray
+	const TArray<FNeuralTensor> InputTensorArray = InOutNetwork->CreateInputArrayCopy();
+	FString TensorsAsString = TEXT("Input(s):\n");
+	for (int32 InputIndex = 0; InputIndex < InOutNetwork->GetInputTensorNumber(); ++InputIndex)
+	{
+		TensorsAsString += InOutNetwork->GetInputTensor(InputIndex).ToString() + TEXT("\n");
+	}
+	// Run network
+	InOutNetwork->SetDeviceType(ENeuralDeviceType::CPU);
+	InOutNetwork->Run();
+	// Create output string
+	TensorsAsString += TEXT("\nOutput(s):\n");
+	for (int32 OutputIndex = 0; OutputIndex < InOutNetwork->GetOutputTensorNumber(); ++OutputIndex)
+	{
+		TensorsAsString += InOutNetwork->GetOutputTensor(OutputIndex).ToString() + TEXT("\n");
+	}
+	// Reset memory
+	InOutNetwork->SetInputFromArrayCopy(InputTensorArray); // This is doing an unnecessary copy (it should move really)
+	// Return TensorsAsString
+	return TensorsAsString + TEXT("\n\n\n");
+}
+
+void FNeuralNetworkInferenceQAOperatorAsset::RunAndAddTest(UNeuralNetwork* InOutNetwork)
+{
+	// Resize NewTests
+	NewTests.Emplace();
+
+	FNeuralNetworkInferenceQAOperatorTestAsset& NewTest = NewTests.Last();
+	NewTestsString += RunNetworkCPUAndGetString(InOutNetwork);
+	NewTest.AddInputTensors(InOutNetwork);
+	NewTest.AddOutputTensors(InOutNetwork);
+}
+
+bool FNeuralNetworkInferenceQAOperatorAsset::CompareNewVsPreviousTests(const FString& InGroundTruthDirectory, const FString& InOperatorName)
+{
+	bool bWasComparisonSuccessful = true;
+	// If PreviousTests is empty --> It is the first time it is being used
+	if (PreviousTests.IsEmpty() || PreviousTests.Num() != NewTests.Num())
+	{
+		if (PreviousTests.IsEmpty())
+		{
+			UE_LOG(LogNeuralNetworkInferenceQA, Warning, TEXT("First time that this test is generated, no previous results exist."));
+		}
+		else
+		{
+			UE_LOG(LogNeuralNetworkInferenceQA, Warning, TEXT("New tests added or removed, PreviousTests.Num() = %d and NewTests.Num() = %d."), PreviousTests.Num(), NewTests.Num());
+		}
+		bWasComparisonSuccessful = false;
+	}
+	// Everything has to pass to be a successful comparison
+	else // if (PreviousTests.Num() == NewTests.Num())
+	{
+		for (int32 TestIndex = 0; TestIndex < PreviousTests.Num(); ++TestIndex)
+		{
+			// Easy but CPU-and-compiler dependent
+			if (PreviousTestsString != NewTestsString)
+			{
+				// Even if PreviousTestsString != NewTestsString, it might be a floating value precision error
+				const FNeuralNetworkInferenceQAOperatorTestAsset& Test = PreviousTests[TestIndex];
+				const FNeuralNetworkInferenceQAOperatorTestAsset& UpdatedTest = NewTests[TestIndex];
+				bWasComparisonSuccessful &= FNeuralNetworkInferenceQAOperatorTestAsset::CompareAverageL1DiffNewVsPreviousTests(Test, UpdatedTest, InOperatorName);
+			}
+		}
+	}
+	// Save on disk
+	if (!bWasComparisonSuccessful)
+	{
+		const FString BaseTestPath = InGroundTruthDirectory / TEXT("temp_") + InOperatorName;
+		const FString GroundTruthFileExtension = TEXT(".txt");
+		const FString FilePathPreviousTest = BaseTestPath + TEXT("_previous") + GroundTruthFileExtension;
+		const FString FilePathNewTest = BaseTestPath + TEXT("_new") + GroundTruthFileExtension;
+		ensureMsgf(FFileHelper::SaveStringToFile(PreviousTestsString, *FilePathPreviousTest), TEXT("FFileHelper::SaveStringToFile returned false."));
+		ensureMsgf(FFileHelper::SaveStringToFile(NewTestsString, *FilePathNewTest), TEXT("FFileHelper::SaveStringToFile returned false."));
+		UE_LOG(LogNeuralNetworkInferenceQA, Warning, TEXT("FNeuralNetworkInferenceQAOperatorAsset::CompareNewVsPreviousTests(): Mismatch between expected and actual results, they should match. Check the following files for differences:\n"
+			"\t- Character length (saved previous vs. new): %d vs. %d\n"
+			"\t- Previous results saved in %s\n"
+			"\t- New results saved in %s\n"),
+			PreviousTestsString.Len(), NewTestsString.Len(),
+			*IFileManager::Get().ConvertToAbsolutePathForExternalAppForRead(*FilePathPreviousTest),
+			*IFileManager::Get().ConvertToAbsolutePathForExternalAppForRead(*FilePathNewTest));
+		// Final swap
+		Swap(PreviousTests, NewTests);
+		Swap(PreviousTestsString, NewTestsString);
+	}
+	// Flush
+	FlushNewTests();
+	// Return bWasComparisonSuccessful
+	return bWasComparisonSuccessful;
+}
+
+void FNeuralNetworkInferenceQAOperatorAsset::FlushNewTests()
+{
+	NewTests.Empty();
+	NewTestsString.Empty();
+}
+
+
+
+/* UNeuralNetworkInferenceQAAsset functions
+ *****************************************************************************/
+
+void UNeuralNetworkInferenceQAAsset::FindOrAddOperators(const TArray<FString>& InOperatorNames)
+{
+	FString NotFoundOperatorNames;
+	for (const FString& OperatorName : InOperatorNames)
+	{
+		if (!Operators.Find(OperatorName))
+		{
+			NotFoundOperatorNames += OperatorName + TEXT(", ");
+			Operators.Add(OperatorName, FNeuralNetworkInferenceQAOperatorAsset());
+			UE_LOG(LogNeuralNetworkInferenceQA, Warning, TEXT("QA for operator %s was not found on UNeuralNetworkInferenceQAAsset, added!"), *OperatorName);
+		}
+	}
+	if (NotFoundOperatorNames.Len() > 0)
+	{
+		NotFoundOperatorNames.LeftChopInline(2); // Removes the last TEXT(", ")
+		ensureMsgf(false, TEXT("Some operators are new and were not found on UNeuralNetworkInferenceQAAsset, they have been added: %s."), *NotFoundOperatorNames);
+	}
+}
+
+void UNeuralNetworkInferenceQAAsset::RunAndAddTest(UNeuralNetwork* InOutNetwork, const FString& OperatorName)
+{
+	Operators[OperatorName].RunAndAddTest(InOutNetwork);
+}
+
+bool UNeuralNetworkInferenceQAAsset::CompareNewVsPreviousTests(const FString& InGroundTruthDirectory)
+{
+	bool bWasComparisonSuccessful = true;
+	for (auto& OperatorPair : Operators)
+	{
+		// No break so we can return all wrong results at once
+		if (!OperatorPair.Value.CompareNewVsPreviousTests(InGroundTruthDirectory, OperatorPair.Key))
+		{
+			bWasComparisonSuccessful = false;
+		}
+	}
+	if (!bWasComparisonSuccessful)
+	{
+		UE_LOG(LogNeuralNetworkInferenceQA, Warning, TEXT("UNeuralNetworkInferenceQAAsset::CompareNewVsPreviousTests(): Mismatch between expected and actual results, they should match. Check the previous warning messages."));
+	}
+	return bWasComparisonSuccessful;
+}
+
+void UNeuralNetworkInferenceQAAsset::FlushNewTests()
+{
+	for (auto& OperatorPair : Operators)
+	{
+		OperatorPair.Value.FlushNewTests();
+	}
+}
+
+const FString UNeuralNetworkInferenceQAAsset_FileName = TEXT("NeuralNetworkInferenceQAAsset");
+
+UNeuralNetworkInferenceQAAsset* UNeuralNetworkInferenceQAAsset::Load(const FString& InNeuralNetworkInferenceQAAssetParentDirectoryName, const FString& InNeuralNetworkInferenceQAAssetName)
+{
+	// Load NeuralNetworkInferenceQAAsset from disk
+	const FString FilePath = UNeuralNetworkInferenceQAAsset_FileName + TEXT("'/Game/") / InNeuralNetworkInferenceQAAssetParentDirectoryName / UNeuralNetworkInferenceQAAsset_FileName + TEXT(".") + UNeuralNetworkInferenceQAAsset_FileName + TEXT("'"); // "NeuralNetworkInferenceQAAsset'/Game/[UNIT_TEST_DIR]/NeuralNetworkInferenceQAAsset.NeuralNetworkInferenceQAAsset'"
+	UNeuralNetworkInferenceQAAsset* NeuralNetworkInferenceQAAsset = LoadObject<UNeuralNetworkInferenceQAAsset>((UObject*)GetTransientPackage(), *FilePath);
+	// Warning if it does not exist yet
+	if (!NeuralNetworkInferenceQAAsset)
+	{
+		UE_LOG(LogNeuralNetworkInferenceQA, Display, TEXT("NeuralNetworkInferenceQAAsset not found in %s. Please, create it first or make sure the right path is being used."), *FilePath);
+		return nullptr;
+	}
+	// Return NeuralNetworkInferenceQAAsset
+	return NeuralNetworkInferenceQAAsset;
+}
+
+bool UNeuralNetworkInferenceQAAsset::Save()
+{
+	// Flush
+	FlushNewTests();
+	// Extract Package
+	UPackage* Package = GetOutermost();
+	if (!Package)
+	{
+		UE_LOG(LogNeuralNetworkInferenceQA, Warning, TEXT("Package was a nullptr!"));
+		return false;
+	}
+	Package->FullyLoad();
+	// Save Uasset
+	FString FilePath;
+	{
+		FString ParentPath, FileName, FileExtension;
+		FPaths::Split(Package->GetLoadedPath().GetLocalFullPath(), ParentPath, FileName, FileExtension); // "../../../Sandbox/MachineLearning/NNIUnitTest/Content/NeuralNetworkInferenceQAAsset.uasset"
+		FilePath = ParentPath / UNeuralNetworkInferenceQAAsset_FileName + TEXT("_new.") + FileExtension; // "../../../Sandbox/MachineLearning/NNIUnitTest/Content/NeuralNetworkInferenceQAAsset_new.uasset"
+	}
+	FSavePackageArgs SaveArgs;
+	SaveArgs.TopLevelFlags = RF_Standalone;
+	const bool bWasPackageSaved = UPackage::SavePackage(Package, /*InNeuralNetworkInferenceQAAsset*/nullptr, *FilePath, SaveArgs);
+	if (!bWasPackageSaved)
+	{
+		UE_LOG(LogNeuralNetworkInferenceQA, Display, TEXT("UNeuralNetworkInferenceQAAsset::Save() failed for %s!"), *FilePath);
+	}
+	// Return success bool
+	return bWasPackageSaved;
+}
diff --git a/Engine/Plugins/Experimental/NeuralNetworkInference/Source/Private/NeuralNetworkInferenceQA/Private/NeuralNetworkInferenceQAUtils.cpp b/Engine/Plugins/Experimental/NeuralNetworkInference/Source/Private/NeuralNetworkInferenceQA/Private/NeuralNetworkInferenceQAUtils.cpp
new file mode 100644
--- /dev/null
+++ b/Engine/Plugins/Experimental/NeuralNetworkInference/Source/Private/NeuralNetworkInferenceQA/Private/NeuralNetworkInferenceQAUtils.cpp
@@ -0,0 +1,124 @@
+// Copyright Epic Games, Inc. All Rights Reserved.
+
+#include "NeuralNetworkInferenceQAUtils.h"
+
+
+
+/* FNeuralNetworkInferenceQAUtils static public functions
+ *****************************************************************************/
+
+bool FNeuralNetworkInferenceQAUtils::EstimateTensorL1DiffError(const FNeuralTensor& InTensorA, const FNeuralTensor& InTensorB, const float InZeroThreshold, const FString& InDebugName)
+{
+	if (InTensorA.GetDataType() != InTensorB.GetDataType())
+	{
+		UE_LOG(LogNeuralNetworkInferenceQA, Warning, TEXT("InTensorA.GetDataType() != InTensorB.GetDataType() failed, %d != %d."), InTensorA.GetDataType(), InTensorB.GetDataType());
+		return false;
+	}
+	if (InTensorA.Num() != InTensorB.Num())
+	{
+		UE_LOG(LogNeuralNetworkInferenceQA, Warning, TEXT("InTensorA.Num() == InTensorB.Num() failed, %d != %d."), InTensorA.Num(), InTensorB.Num());
+		return false;
+	}
+	const int32 TensorSize = InTensorA.Num();
+	float AverageL1Value = 0.f;
+	float AverageL1Diff = 0.f;
+	int32 SizeMinusInvalid = TensorSize;
+	FString FiniteVsNon;
+	// Float
+	const ENeuralDataType DataType = InTensorA.GetDataType();
+	if (DataType == ENeuralDataType::Float)
+	{
+		for (int32 Index = 0; Index < TensorSize; ++Index)
+		{
+			const float Value1 = InTensorA.At<float>(Index);
+			const float Value2 = InTensorB.At<float>(Index);
+			const bool bIsFinite1 = FMath::IsFinite(Value1);
+			const bool bIsFinite2 = FMath::IsFinite(Value2);
+			// Update L1
+			if (bIsFinite1 && bIsFinite2)
+			{
+				AverageL1Value += FMath::Abs(Value1);
+				AverageL1Diff += FMath::Abs(Value1 - Value2);
+			}
+			// Update NaN|Inf stats
+			else
+			{
+				if (bIsFinite1 != bIsFinite2)
+				{
+					FiniteVsNon += (bIsFinite1 ? FString::SanitizeFloat(Value1) : FString(TEXT("NaNInf"))) + TEXT("-")
+						+ (bIsFinite2 ? FString::SanitizeFloat(Value2) : FString(TEXT("NaNInf"))) + TEXT("; ");
+				}
+				--SizeMinusInvalid;
+			}
+		}
+	}
+	// Integer types
+	else
+	{
+		for (int32 Index = 0; Index < TensorSize; ++Index)
+		{
+			float Value1;
+			float Value2;
+			if (DataType == ENeuralDataType::Int32)
+			{
+				Value1 = InTensorA.At<int32>(Index);
+				Value2 = InTensorB.At<int32>(Index);
+			}
+			else if (DataType == ENeuralDataType::Int64)
+			{
+				Value1 = InTensorA.At<int64>(Index);
+				Value2 = InTensorB.At<int64>(Index);
+			}
+			else if (DataType == ENeuralDataType::UInt32)
+			{
+				Value1 = InTensorA.At<uint32>(Index);
+				Value2 = InTensorB.At<uint32>(Index);
+			}
+			else if (DataType == ENeuralDataType::UInt64)
+			{
+				Value1 = InTensorA.At<uint64>(Index);
+				Value2 = InTensorB.At<uint64>(Index);
+			}
+			else
+			{
+				Value1 = 0.f;
+				Value2 = 0.f;
+				ensureMsgf(false, TEXT("DataType = %d not implemented yet for FNeuralNetworkInferenceQAUtils::EstimateTensorL1DiffError()!"), InTensorA.GetDataType());
+			}
+			// Update L1
+			AverageL1Value += FMath::Abs(Value1);
+			AverageL1Diff += FMath::Abs(Value1 - Value2);
+		}
+	}
+
+	// Were all numbers NaN|Inf?
+	if (SizeMinusInvalid < 1)
+	{
+		UE_LOG(LogNeuralNetworkInferenceQA, Display, TEXT("All %d numbers were NaN|Inf."), TensorSize);
+		return true;
+	}
+	// Some verbose
+	if (FiniteVsNon.Len() > 0)
+	{
+		UE_LOG(LogNeuralNetworkInferenceQA, Display, TEXT("IsFinite(Value1) != IsFinite(Value2) for %d values (skipped %d valid ones): %s"), TensorSize - SizeMinusInvalid, SizeMinusInvalid, *FiniteVsNon);
+	}
+	// Average L1
+	AverageL1Value /= SizeMinusInvalid;
+	AverageL1Diff /= (SizeMinusInvalid * (AverageL1Value > 0.01f ? AverageL1Value : 1)); // Relative if > 0.01, absolute otherwise
+	// Sanity check
+	if (!FMath::IsFinite(AverageL1Value) || !FMath::IsFinite(AverageL1Diff))
+	{
+		UE_LOG(LogNeuralNetworkInferenceQA, Warning, TEXT("FNeuralNetworkInferenceQAUtils::EstimateTensorL1DiffError(): NaN|Inf found: AverageL1Diff %f for an L1 average value of %f."), AverageL1Diff, AverageL1Value);
+		return false;
+	}
+	// Did it pass?
+	if (AverageL1Diff > InZeroThreshold)
+	{
+		UE_LOG(LogNeuralNetworkInferenceQA, Warning, TEXT("%s: AverageL1Diff <= ZeroThreshold failed (%fe-3 > %fe-3), AverageL1Value = %f, AverageL1Diff/ZeroThreshold = %f."),
+			*InDebugName, 1e3 * AverageL1Diff, 1e3 * InZeroThreshold, AverageL1Value, AverageL1Diff / InZeroThreshold);
+		UE_LOG(LogNeuralNetworkInferenceQA, Warning, TEXT("- TensorA = %s"), *InTensorA.ToString(100));
+		UE_LOG(LogNeuralNetworkInferenceQA, Warning, TEXT("- TensorB = %s"), *InTensorB.ToString(100));
+		return false;
+	}
+	return true;
+}
diff --git a/Engine/Plugins/Experimental/NeuralNetworkInference/Source/Private/NeuralNetworkInferenceQA/Private/NeuralNetworkInferenceQAUtils.h b/Engine/Plugins/Experimental/NeuralNetworkInference/Source/Private/NeuralNetworkInferenceQA/Private/NeuralNetworkInferenceQAUtils.h
new file mode 100644
--- /dev/null
+++ b/Engine/Plugins/Experimental/NeuralNetworkInference/Source/Private/NeuralNetworkInferenceQA/Private/NeuralNetworkInferenceQAUtils.h
@@ -0,0 +1,18 @@
+// Copyright Epic Games, Inc. All Rights Reserved.
+
+#pragma once
+
+#include "CoreMinimal.h"
+#include "NeuralTensor.h"
+
+DEFINE_LOG_CATEGORY_STATIC(LogNeuralNetworkInferenceQA, Display, All);
+
+class FNeuralNetworkInferenceQAUtils
+{
+public:
+	/**
+	 * It calculates the error between 2 data arrays (it should return 0 if both arrays are the same).
+	 * @return False if the estimated error > InZeroThreshold, true otherwise.
+	 */
+	static bool EstimateTensorL1DiffError(const FNeuralTensor& InTensorA, const FNeuralTensor& InTensorB, const float InZeroThreshold, const FString& InDebugName);
+};
diff --git a/Engine/Plugins/Experimental/NeuralNetworkInference/Source/Private/NeuralNetworkInferenceQA/Private/OperatorTester.cpp b/Engine/Plugins/Experimental/NeuralNetworkInference/Source/Private/NeuralNetworkInferenceQA/Private/OperatorTester.cpp
new file mode 100644
--- /dev/null
+++ b/Engine/Plugins/Experimental/NeuralNetworkInference/Source/Private/NeuralNetworkInferenceQA/Private/OperatorTester.cpp
@@ -0,0 +1,1224 @@
+// Copyright Epic Games, Inc. All Rights Reserved.
+
+#include "OperatorTester.h"
+#include "NeuralNetworkInferenceQAUtils.h"
+#include "NeuralOperators.h"
+#include "Misc/Paths.h"
+
+
+
+/* FOperatorTester static public functions
+ *****************************************************************************/
+
+void FOperatorTester::TestOperator(UNeuralNetworkInferenceQAAsset* InOutNetworkInferenceQAAsset, TArray<TArray<FNeuralTensor>>& InOutTensors, TArray<TArray<FNeuralTensor*>>& InOutInputTensors,
+	TArray<TArray<FNeuralTensor*>>& InOutOutputTensors, TArray<TSharedPtr<FNeuralOperator>>& InOutOperators, const FString& InOperatorName, const FString& InGroundTruthDirectory,
+	const float InZeroThreshold, const bool bInShouldOperatorsProvideSameResult, const bool bInShouldRunGPU, const bool bInShouldRunFinalStringVsIOTest)
+{
+	// Sanity checks
+	if (!InOutNetworkInferenceQAAsset)
+	{
+		ensureMsgf(false, TEXT("InOutNetworkInferenceQAAsset was nullptr!"));
+		return;
+	}
+	else if (InOutInputTensors.Num() != InOutOperators.Num())
+	{
+		ensureMsgf(false, TEXT("InOutInputTensors.Num() == InOutOperators.Num(): %d vs. %d"), InOutInputTensors.Num(), InOutOperators.Num());
+		return;
+	}
+	if (!bInShouldRunGPU)
+	{
+		UE_LOG(LogNeuralNetworkInferenceQA, Warning, TEXT("FOperatorTester::TestOperator(): bInShouldRunGPU is false, this should only happen for very early deployment of new operators."));
+	}
+	// Run test on each input FNeuralTensor
+	FString Index0StringIfShouldProvideSameResult;
+	for (int32 Index = 0; Index < InOutInputTensors.Num(); ++Index)
+	{
+		//UE_LOG(LogNeuralNetworkInferenceQA, Display, TEXT("------------------------- Subtest #%d/%d"), Index+1, InOutInputTensors.Num());
+		TSharedPtr<FNeuralOperator>& Operator = InOutOperators[Index];
+		const TArray<FNeuralTensor*>& InputTensors = InOutInputTensors[Index];
+		if (InputTensors.IsEmpty())
+		{
+			ensureMsgf(false, TEXT("InputTensors.Num() is 0 but should be > 0."));
+			return;
+		}
+		// Operator array
+		const int32 InlinedTensorIndex = Operator->SetInputTensorsAndGetInlinedIndex(InputTensors);
+		// Tensor manager
+		const TArray<FNeuralTensor*>& OutputTensors = (InlinedTensorIndex > -1 ? TArray<FNeuralTensor*>({ InputTensors[InlinedTensorIndex] }) : InOutOutputTensors[Index]);
+		if (OutputTensors.IsEmpty())
+		{
+			ensureMsgf(false, TEXT("OutputTensors.Num() is 0 but should be > 0."));
+			return;
+		}
+		// Operator array
+		if (InlinedTensorIndex < 0)
+		{
+			Operator->SetOutputTensors(InOutOutputTensors[Index]);
+		}
+		// Create and add auxiliary tensors for that operator (if any)
+		const int64 NumberAuxiliaryTensors = Operator->GetNumberAuxiliaryTensors();
+		if (NumberAuxiliaryTensors > 0)
+		{
+			TArray<FNeuralTensor*> AuxiliaryOperatorTensors;
+			for (int64 AuxiliaryTensorIndex = 0; AuxiliaryTensorIndex < Operator->GetNumberAuxiliaryTensors(); ++AuxiliaryTensorIndex)
+			{
+				const FString AuxiliaryTensorName = Operator->GetName() + TEXT("-Aux") + FString::FromInt(AuxiliaryTensorIndex);
+				InOutTensors[Index].Push(FNeuralTensor(ENeuralDataType::None, TArray<int64>({}), AuxiliaryTensorName, ENeuralTensorType::Generic)); // ENeuralTensorType can be modify by FNeuralOperator::Configure()
+				AuxiliaryOperatorTensors.Push(&InOutTensors[Index].Last());
+			}
+			Operator->SetAuxiliaryTensors(AuxiliaryOperatorTensors);
+		}
+		// Configure operator
+		if (!Operator->Configure())
+		{
+			ensureMsgf(false, TEXT("Operator %s could not be configured."), *Operator->GetName());
+			return;
+		}
+		// Network tests
+		UNeuralNetwork* Network = NewObject<UNeuralNetwork>((UObject*)GetTransientPackage(), UNeuralNetwork::StaticClass());
+		if (!Network)
+		{
+			ensureMsgf(false, TEXT("Network is nullptr."));
+			return;
+		}
+		if (!Network->Load(InOutTensors[Index], InputTensors, OutputTensors, { Operator }))
+		{
+			ensureMsgf(false, TEXT("Network could not be loaded."));
+			return;
+		}
+		// GPU vs. CPU test (if enabled)
+		if (bInShouldRunGPU)
+		{
+			NetworkGPUvsCPU(Network, InlinedTensorIndex > -1, InZeroThreshold);
+		}
+		Network->SetDeviceType(ENeuralDeviceType::CPU);
+		if (InlinedTensorIndex > -1)
+		{
+			NetworkInputVsOutput(Network, InlinedTensorIndex);
+		}
+		// Test against GT saved on disk
+		// Case a) Only saving first test, the others should match the first test results
+		if (bInShouldOperatorsProvideSameResult)
+		{
+			if (Index == 0)
+			{
+				InOutNetworkInferenceQAAsset->RunAndAddTest(Network, InOperatorName);
+				Index0StringIfShouldProvideSameResult = FNeuralNetworkInferenceQAOperatorAsset::RunNetworkCPUAndGetString(Network);
+			}
+			else
+			{
+				// Sanity check
+				const FString OutputTensorsAsString = FNeuralNetworkInferenceQAOperatorAsset::RunNetworkCPUAndGetString(Network);
+				if (Index0StringIfShouldProvideSameResult != OutputTensorsAsString)
+				{
+					UE_LOG(LogNeuralNetworkInferenceQA, Display, TEXT("Index0StringIfShouldProvideSameResult:\n%s"), *Index0StringIfShouldProvideSameResult);
+					UE_LOG(LogNeuralNetworkInferenceQA, Display, TEXT("OutputTensorsAsString:\n%s"), *OutputTensorsAsString);
+					ensureMsgf(false, TEXT("Subtest #%d did not match subtest #0"), Index);
+				}
+			}
+		}
+		// Case b) Saving each sub-iteration
+		else
+		{
+			InOutNetworkInferenceQAAsset->RunAndAddTest(Network, InOperatorName);
+		}
+	}
+}
+
+void FOperatorTester::TestBatchNormalizationOperator(UNeuralNetworkInferenceQAAsset* InOutNetworkInferenceQAAsset, const TMap<int32, TArray<float>>& InInputTensorData, const TMap<int32, TArray<float>>& InInputTensorDataSmall,
+	const TMap<int32, TArray<float>>& InInputTensorDataPos, const float InZeroThreshold, const FString& InGroundTruthDirectory, const int32 InOperatorCounter, const int32 InOperatorTotal)
+{
+	UE_LOG(LogNeuralNetworkInferenceQA, Display, TEXT("--------------- %d/%d FBatchNormalizationOperator"), InOperatorCounter, InOperatorTotal);
+	// Inout tensors (X != Mean to avoid Output = Bias; Variance positive to avoid NaN)
+	TArray<TArray<FNeuralTensor>> Tensors({
+		// ----------------------------------------------------------------- ONNX - [REDACTED_PROFILE_URL]/onnx/blob/master/docs/TestCoverage.md#BatchNormalization -----------------------------------------------------------------
+		{FNeuralTensor(TArray<float>({-1,0,1, 2,3,4}), TArray<int64>({1,2,1,3})),
+			FNeuralTensor(TArray<float>({1, 1.5}), TArray<int64>({2})),
+			FNeuralTensor(TArray<float>({0, 1}), TArray<int64>({2})),
+			FNeuralTensor(TArray<float>({0, 3}), TArray<int64>({2})),
+			FNeuralTensor(TArray<float>({1, 1.5}), TArray<int64>({2})),
+			FNeuralTensor(ENeuralDataType::Float, {1,2,1,3})}, // [-0.999995 0.0 0.999995; -0.22474074 1.0 2.2247407]
+		{FNeuralTensor(InInputTensorData.FindChecked(120), TArray<int64>({2,3,4,5})),
+			FNeuralTensor(InInputTensorData.FindChecked(3), TArray<int64>({3})),
+			FNeuralTensor(InInputTensorData.FindChecked(3), TArray<int64>({3})),
+			FNeuralTensor(InInputTensorData.FindChecked(3), TArray<int64>({3})),
+			FNeuralTensor(InInputTensorDataPos.FindChecked(3), TArray<int64>({3})),
+			FNeuralTensor(ENeuralDataType::Float, {1,2,1,3})},
+		// Scalars
+		{FNeuralTensor(InInputTensorData.FindChecked(1)), FNeuralTensor(InInputTensorData.FindChecked(1)), FNeuralTensor(InInputTensorData.FindChecked(1)),
+			FNeuralTensor(InInputTensorDataSmall.FindChecked(1)), FNeuralTensor(InInputTensorDataPos.FindChecked(1)), FNeuralTensor(ENeuralDataType::Float, 1)},
+		// Scalars
+		{FNeuralTensor(InInputTensorData.FindChecked(1)), FNeuralTensor(InInputTensorData.FindChecked(1)), FNeuralTensor(InInputTensorData.FindChecked(1)),
+			FNeuralTensor(InInputTensorDataSmall.FindChecked(1)), FNeuralTensor(InInputTensorDataPos.FindChecked(1)), FNeuralTensor(ENeuralDataType::Float, 1)},
+		// 1-D
+		{FNeuralTensor(InInputTensorData.FindChecked(10)), FNeuralTensor(InInputTensorData.FindChecked(1)), FNeuralTensor(InInputTensorData.FindChecked(1)),
+			FNeuralTensor(InInputTensorDataSmall.FindChecked(1)), FNeuralTensor(InInputTensorDataPos.FindChecked(1)), FNeuralTensor(ENeuralDataType::Float, 10)},
+		// 4-D with 1 batch and 1 channel
+		{FNeuralTensor(InInputTensorData.FindChecked(100), TArray<int64>({1, 1, 5, 20})), FNeuralTensor(InInputTensorData.FindChecked(1)),
+			FNeuralTensor(InInputTensorData.FindChecked(1)), FNeuralTensor(InInputTensorDataSmall.FindChecked(1)),
+			FNeuralTensor(InInputTensorDataPos.FindChecked(1)), FNeuralTensor(ENeuralDataType::Float, {1, 1, 5, 20})},
+		// 4-D with 1 channel
+		{FNeuralTensor(InInputTensorData.FindChecked(150), TArray<int64>({5, 1, 6, 5})), FNeuralTensor(InInputTensorData.FindChecked(1)),
+			FNeuralTensor(InInputTensorData.FindChecked(1)), FNeuralTensor(InInputTensorDataSmall.FindChecked(1)),
+			FNeuralTensor(InInputTensorDataPos.FindChecked(1)), FNeuralTensor(ENeuralDataType::Float, {5, 1, 6, 5})},
+		// 4-D with 1 batch
+		{FNeuralTensor(InInputTensorData.FindChecked(150), TArray<int64>({1, 5, 6, 5})), FNeuralTensor(InInputTensorData.FindChecked(5)),
+			FNeuralTensor(InInputTensorData.FindChecked(5)), FNeuralTensor(InInputTensorDataSmall.FindChecked(5)),
+			FNeuralTensor(InInputTensorDataPos.FindChecked(5)), FNeuralTensor(ENeuralDataType::Float, {1, 5, 6, 5})},
+		// 4-D
+		{FNeuralTensor(InInputTensorData.FindChecked(300), TArray<int64>({3, 10, 5, 2})), FNeuralTensor(InInputTensorData.FindChecked(10)),
+			FNeuralTensor(InInputTensorData.FindChecked(10)), FNeuralTensor(InInputTensorDataSmall.FindChecked(10)),
+			FNeuralTensor(InInputTensorDataPos.FindChecked(10)), FNeuralTensor(ENeuralDataType::Float, {3, 10, 5, 2})},
+		// 2-D
+		{FNeuralTensor(InInputTensorData.FindChecked(18), TArray<int64>({6, 3})), FNeuralTensor(InInputTensorData.FindChecked(3)),
+			FNeuralTensor(InInputTensorData.FindChecked(3)), FNeuralTensor(InInputTensorDataSmall.FindChecked(3)),
+			FNeuralTensor(InInputTensorDataPos.FindChecked(3)), FNeuralTensor(ENeuralDataType::Float, {6, 3})},
+		// 3-D
+		{FNeuralTensor(InInputTensorData.FindChecked(30), TArray<int64>({3, 5, 2})), FNeuralTensor(InInputTensorData.FindChecked(5)),
+			FNeuralTensor(InInputTensorData.FindChecked(5)), FNeuralTensor(InInputTensorDataSmall.FindChecked(5)),
+			FNeuralTensor(InInputTensorDataPos.FindChecked(5)), FNeuralTensor(ENeuralDataType::Float, {3, 5, 2})},
+		// 5-D
+		{FNeuralTensor(InInputTensorData.FindChecked(300), TArray<int64>({2, 3, 5, 5, 2})), FNeuralTensor(InInputTensorData.FindChecked(3)),
+			FNeuralTensor(InInputTensorData.FindChecked(3)), FNeuralTensor(InInputTensorDataSmall.FindChecked(3)),
+			FNeuralTensor(InInputTensorDataPos.FindChecked(3)), FNeuralTensor(ENeuralDataType::Float, {2, 3, 5, 5, 2})}
+		});
+	TArray<TArray<FNeuralTensor*>> InputTensors;
+	TArray<TArray<FNeuralTensor*>> OutputTensors;
+	TArray<TSharedPtr<FNeuralOperator>> Operators;
+	for (int32 TensorIndex = 0; TensorIndex < Tensors.Num(); ++TensorIndex)
+	{
+		InputTensors.Push({ &Tensors[TensorIndex][0], &Tensors[TensorIndex][1], &Tensors[TensorIndex][2], &Tensors[TensorIndex][3], &Tensors[TensorIndex][4] });
+		OutputTensors.Push(TensorIndex % 2 == 0 ? TArray<FNeuralTensor*>({ &Tensors[TensorIndex][5] }) : TArray<FNeuralTensor*>({/*Inlined*/ }));
+		Operators.Push(MakeShared<FBatchNormalizationOperator>(OutputTensors[TensorIndex].Num() == 0,
+			/*Epsilon*/(TensorIndex == 1 ? 1e-2 : TensorIndex == 3 ? 10 : TensorIndex == 6 ? 1 : 1e-5),
+			/*Momentum*/(TensorIndex+1.f)/Tensors.Num()));
+	}
+	TestOperator(InOutNetworkInferenceQAAsset, Tensors, InputTensors, OutputTensors, Operators, TEXT("FBatchNormalizationOperator"), InGroundTruthDirectory, InZeroThreshold,
+		/*bShouldOperatorsProvideSameResult*/false, /*bInShouldRunGPU*/true);
+}
+
+void FOperatorTester::TestGemmOperator(UNeuralNetworkInferenceQAAsset* InOutNetworkInferenceQAAsset, const TMap<int32, TArray<float>>& InInputTensorData, const float InZeroThreshold, const FString& InGroundTruthDirectory,
+	const int32 InOperatorCounter, const int32 InOperatorTotal)
+{
+	UE_LOG(LogNeuralNetworkInferenceQA, Display, TEXT("--------------- %d/%d FGemmOperator"), InOperatorCounter, InOperatorTotal);
+	// Inout tensors
+	const TArray<TArray<FNeuralTensor>> GemmABCOutputTensorTests({
+		// Scalars
+		{FNeuralTensor(InInputTensorData.FindChecked(1)), FNeuralTensor(InInputTensorData.FindChecked(1)), FNeuralTensor(InInputTensorData.FindChecked(1)),
+			FNeuralTensor(ENeuralDataType::Float, 1)},
+		// 1-D
+		{FNeuralTensor(InInputTensorData.FindChecked(300), TArray<int64>({50, 6})), FNeuralTensor(InInputTensorData.FindChecked(18), TArray<int64>({6, 3})),
+			FNeuralTensor(InInputTensorData.FindChecked(150), TArray<int64>({50, 3})), FNeuralTensor(ENeuralDataType::Float, TArray<int64>({50, 3}))},
+		// 1-D, C size = [(1), 3], the next 2 should give the exact same results
+		{FNeuralTensor(InInputTensorData.FindChecked(300), TArray<int64>({50, 6})), FNeuralTensor(InInputTensorData.FindChecked(18), TArray<int64>({6, 3})),
+			FNeuralTensor(InInputTensorData.FindChecked(3), TArray<int64>({1, 3})), FNeuralTensor(ENeuralDataType::Float, TArray<int64>({50, 3}))}, // C size = [1, 3]
+		{FNeuralTensor(InInputTensorData.FindChecked(300), TArray<int64>({50, 6})), FNeuralTensor(InInputTensorData.FindChecked(18), TArray<int64>({6, 3})),
+			FNeuralTensor(InInputTensorData.FindChecked(3)), FNeuralTensor(ENeuralDataType::Float, TArray<int64>({50, 3}))}, // C size = [3]
+		// 1-D, C size = [50, 1]
+		{FNeuralTensor(InInputTensorData.FindChecked(300), TArray<int64>({50, 6})), FNeuralTensor(InInputTensorData.FindChecked(18), TArray<int64>({6, 3})),
+			FNeuralTensor(InInputTensorData.FindChecked(50), TArray<int64>({50, 1})), FNeuralTensor(ENeuralDataType::Float, TArray<int64>({50, 3}))}
+	});
+	for (int32 TensorTestsIndex = 0; TensorTestsIndex < GemmABCOutputTensorTests.Num(); ++TensorTestsIndex)
+	{
+		UE_LOG(LogNeuralNetworkInferenceQA, Display, TEXT("-------------------- FGemmOperator #%d/%d"), TensorTestsIndex + 1, GemmABCOutputTensorTests.Num());
+		TArray<TSharedPtr<FNeuralOperator>> Operators({
+			// A, B
+			MakeShared<FGemmOperator>(/*Alpha*/-0.5, /*Beta*/-10, /*bTransA*/false, /*bTransB*/false),// = AB + C
+			MakeShared<FGemmOperator>(/*Alpha*/10, /*Beta*/0, /*bTransA*/false, /*bTransB*/false),    // = AB
+			MakeShared<FGemmOperator>(/*Alpha*/0, /*Beta*/0.1, /*bTransA*/false, /*bTransB*/false),   // = C
+			MakeShared<FGemmOperator>(/*Alpha*/0, /*Beta*/0, /*bTransA*/false, /*bTransB*/false),     // = 0
+		});
+		// 2nd and 3rd set of tests are mostly for C, so we can ignore bTransA/bTransB tests.
+		if (TensorTestsIndex < 2)
+		{
+			Operators.Append({
+				// A', B
+				MakeShared<FGemmOperator>(/*Alpha*/1, /*Beta*/1, /*bTransA*/true, /*bTransB*/false),      // = A'B + C
+				MakeShared<FGemmOperator>(/*Alpha*/0.5f, /*Beta*/0, /*bTransA*/true, /*bTransB*/false),   // = 0.5A'B
+				MakeShared<FGemmOperator>(/*Alpha*/0, /*Beta*/0.5f, /*bTransA*/true, /*bTransB*/false),   // = 0.5C
+				MakeShared<FGemmOperator>(/*Alpha*/0, /*Beta*/0, /*bTransA*/true, /*bTransB*/false),      // = 0
+				// A, B'
+				MakeShared<FGemmOperator>(/*Alpha*/1, /*Beta*/1, /*bTransA*/false, /*bTransB*/true),      // = AB' + C
+				MakeShared<FGemmOperator>(/*Alpha*/10, /*Beta*/0, /*bTransA*/false, /*bTransB*/true),     // = 10AB'
+				MakeShared<FGemmOperator>(/*Alpha*/0, /*Beta*/10, /*bTransA*/false, /*bTransB*/true),     // = 10C
+				MakeShared<FGemmOperator>(/*Alpha*/0, /*Beta*/0, /*bTransA*/false, /*bTransB*/true),      // = 0
+				// A', B'
+				MakeShared<FGemmOperator>(/*Alpha*/1, /*Beta*/1, /*bTransA*/true, /*bTransB*/true),       // = A'B' + C
+				MakeShared<FGemmOperator>(/*Alpha*/-5, /*Beta*/0, /*bTransA*/true, /*bTransB*/true),      // = -5A'B'
+				MakeShared<FGemmOperator>(/*Alpha*/0, /*Beta*/-5, /*bTransA*/true, /*bTransB*/true),      // = -5C
+				MakeShared<FGemmOperator>(/*Alpha*/0, /*Beta*/0, /*bTransA*/true, /*bTransB*/true)        // = 0
+			});
+		}
+		TArray<TArray<FNeuralTensor>> GemmABCOutputTensors;
+		GemmABCOutputTensors.Init(GemmABCOutputTensorTests[TensorTestsIndex], Operators.Num());
+		// Input/output tensors
+		TArray<TArray<FNeuralTensor*>> InputTensors;
+		TArray<TArray<FNeuralTensor*>> OutputTensors;
+		for (int32 InoutIndex = 0; InoutIndex < Operators.Num(); ++InoutIndex)
+		{
+			// Disable C for half of the cases where Beta==0 (same result)
+			if (/*Beta == 0*/InoutIndex % 2 == 1 && InoutIndex > 8)
+			{
+				InputTensors.Push({ &GemmABCOutputTensors[InoutIndex][0], &GemmABCOutputTensors[InoutIndex][1] });
+			}
+			else
+			{
+				InputTensors.Push({ &GemmABCOutputTensors[InoutIndex][0], &GemmABCOutputTensors[InoutIndex][1], &GemmABCOutputTensors[InoutIndex][2] });
+			}
+			// Transpose FNeuralTensor if bTransA or bTransB
+			if ((InoutIndex >= 4 && InoutIndex < 8) || InoutIndex >= 12) // bTransA?
+			{
+				if (!InputTensors.Last()[0]->Transpose())
+				{
+					ensureMsgf(false, TEXT("FNeuralTensor::Transpose() failed."));
+				}
+			}
+			if (InoutIndex >= 8) // bTransB?
+			{
+				if (!InputTensors.Last()[1]->Transpose())
+				{
+					ensureMsgf(false, TEXT("FNeuralTensor::Transpose() failed."));
+				}
+			}
+			OutputTensors.Push({ &GemmABCOutputTensors[InoutIndex][3] });
+		}
+		// Run actual layer unit testing
+		TestOperator(InOutNetworkInferenceQAAsset, GemmABCOutputTensors, InputTensors, OutputTensors, Operators, /*FileName*/TEXT("FGemmOperator"), InGroundTruthDirectory, InZeroThreshold,
+			/*bShouldOperatorsProvideSameResult*/false, /*bShouldRunGPU*/true, (TensorTestsIndex+1 == GemmABCOutputTensorTests.Num()));
+	}
+}
+
+void FOperatorTester::TestConvOperator(UNeuralNetworkInferenceQAAsset* InOutNetworkInferenceQAAsset, const TMap<int32, TArray<float>>& InInputTensorData, const TMap<int32, TArray<float>>& InInputTensorDataSmall,
+	const float InZeroThreshold, const FString& InGroundTruthDirectory, const int32 InOperatorCounter, const int32 InOperatorTotal)
+{
+	UE_LOG(LogNeuralNetworkInferenceQA, Display, TEXT("--------------- %d/%d FConvOperator"), InOperatorCounter, InOperatorTotal);
+	// Inout tensors
+	TArray<TArray<FNeuralTensor>> Tensors({
+		// ----------------------------------------------------------------- ONNX - [REDACTED_PROFILE_URL]/onnx/blob/master/docs/TestCoverage.md#Conv -----------------------------------------------------------------
+		// conv
+		{FNeuralTensor(TArray<float>({0,1,2,3,4, 5,6,7,8,9, 10,11,12,13,14, 15,16,17,18,19, 20,21,22,23,24}), TArray<int64>({1,1,5,5})),
+			FNeuralTensor(TArray<float>({1,1,1, 1,1,1, 1,1,1}), TArray<int64>({1,1,3,3})),
+			FNeuralTensor(ENeuralDataType::Float, {1,1,5,5}), // [[[[12.0 21.0 27.0 33.0 24.0], [33.0 54.0 63.0 72.0 51.0], [63.0 99.0 108.0 117.0 81.0], [93.0 144.0 153.0 162.0 111.0], [72.0 111.0 117.0 123.0 84.0]]]]
+			FNeuralTensor(), /*Group*/FNeuralTensor(), /*AutoPad*/FNeuralTensor(), /*Pads*/FNeuralTensor(TArray<float>({1,1,1,1}))},
+		{FNeuralTensor(TArray<float>({0,1,2,3,4, 5,6,7,8,9, 10,11,12,13,14, 15,16,17,18,19, 20,21,22,23,24}), TArray<int64>({1,1,5,5})),
+			FNeuralTensor(TArray<float>({1,1,1, 1,1,1, 1,1,1}), TArray<int64>({1,1,3,3})),
+			FNeuralTensor(ENeuralDataType::Float, {1,1,3,3})}, // [[[[54.0 63.0 72.0], [99.0 108.0 117.0], [144.0 153.0 162.0]]]]
+		// conv_with_autopad_same
+		{FNeuralTensor(TArray<float>({0,1,2,3,4, 5,6,7,8,9, 10,11,12,13,14, 15,16,17,18,19, 20,21,22,23,24}), TArray<int64>({1,1,5,5})),
+			FNeuralTensor(TArray<float>({1,1,1, 1,1,1, 1,1,1}), TArray<int64>({1,1,3,3})),
+			FNeuralTensor(ENeuralDataType::Float, {1,1,3,3}), // [[[[12.0 27.0 24.0], [63.0 108.0 81.0], [72.0 117.0 84.0]]]]
+			FNeuralTensor(), /*Group*/FNeuralTensor(), /*AutoPad*/FNeuralTensor(TArray<float>({float(FConvBaseOperator::EAutoPad::SameLower)})), /*Pads*/FNeuralTensor(),
+			/*Dilations*/FNeuralTensor(), /*Strides*/FNeuralTensor(TArray<float>({2,2}))},
+		// conv_with_strides
+		{FNeuralTensor(TArray<float>({0,1,2,3,4, 5,6,7,8,9, 10,11,12,13,14, 15,16,17,18,19, 20,21,22,23,24, 25,26,27,28,29, 30,31,32,33,34}), TArray<int64>({1,1,7,5})),
+			FNeuralTensor(TArray<float>({1,1,1, 1,1,1, 1,1,1}), TArray<int64>({1,1,3,3})),
+			FNeuralTensor(ENeuralDataType::Float, {1,1,4,3}), // [[[[12.0 27.0 24.0], [63.0 108.0 81.0], [123.0 198.0 141.0], [112.0 177.0 124.0]]]]
+			FNeuralTensor(), /*Group*/FNeuralTensor(), /*AutoPad*/FNeuralTensor(), /*Pads*/FNeuralTensor(TArray<float>({1,1, 1,1})),
+			/*Dilations*/FNeuralTensor(), /*Strides*/FNeuralTensor(TArray<float>({2,2}))},
+		{FNeuralTensor(TArray<float>({0,1,2,3,4, 5,6,7,8,9, 10,11,12,13,14, 15,16,17,18,19, 20,21,22,23,24, 25,26,27,28,29, 30,31,32,33,34}), TArray<int64>({1,1,7,5})),
+			FNeuralTensor(TArray<float>({1,1,1, 1,1,1, 1,1,1}), TArray<int64>({1,1,3,3})),
+			FNeuralTensor(ENeuralDataType::Float, {1,1,3,2}), // [[[[54.0 72.0], [144.0 162.0], [234.0 252.0]]]]
+			FNeuralTensor(), /*Group*/FNeuralTensor(), /*AutoPad*/FNeuralTensor(), /*Pads*/FNeuralTensor(),
+			/*Dilations*/FNeuralTensor(), /*Strides*/FNeuralTensor(TArray<float>({2,2}))},
+		// Output size (thus solution) in onnx website is wrong. Replaced Y size to [3, 2]
+		{FNeuralTensor(TArray<float>({0,1,2,3,4, 5,6,7,8,9, 10,11,12,13,14, 15,16,17,18,19, 20,21,22,23,24, 25,26,27,28,29, 30,31,32,33,34}), TArray<int64>({1,1,7,5})),
+			FNeuralTensor(TArray<float>({1,1,1, 1,1,1, 1,1,1}), TArray<int64>({1,1,3,3})),
+			FNeuralTensor(ENeuralDataType::Float, {1,1,3,2}), // [[[[12.0 27.0], [63.0 108.0], [123.0 198.0]]]]
+			FNeuralTensor(), /*Group*/FNeuralTensor(), /*AutoPad*/FNeuralTensor(), /*Pads*/FNeuralTensor(TArray<float>({1,0, 1,0})),
+			/*Dilations*/FNeuralTensor(), /*Strides*/FNeuralTensor(TArray<float>({2,2}))},
+		// ----------------------------------------------------------------- 1D convolutions -----------------------------------------------------------------
+		// X={1,1,1}, W={1,1,1}, B={1}, Y={1,1,1}
+		{FNeuralTensor(InInputTensorData.FindChecked(1), TArray<int64>({1,1,1})), FNeuralTensor(InInputTensorDataSmall.FindChecked(1), TArray<int64>({1,1,1})),
+			FNeuralTensor(ENeuralDataType::Float, {1,1,1}), FNeuralTensor(InInputTensorDataSmall.FindChecked(1))},
+		// X={1,1,9}, W={1,1,5}, B={}, Y={1,1,5}
+		{FNeuralTensor(InInputTensorData.FindChecked(9), TArray<int64>({1,1,9})), FNeuralTensor(InInputTensorDataSmall.FindChecked(5), TArray<int64>({1,1,5})),
+			FNeuralTensor(ENeuralDataType::Float, {1,1,5})},
+		// X={1,3,27}, W={1,3,5}, B={}, Y={1,1,23}
+		{FNeuralTensor(InInputTensorData.FindChecked(81), TArray<int64>({1,3,27})), FNeuralTensor(InInputTensorDataSmall.FindChecked(15), TArray<int64>({1,3,5})),
+			FNeuralTensor(ENeuralDataType::Float, {1,1,23})},
+		// X={1,3,27}, W={2,3,5}, B={}, Y={1,2,23}
+		{FNeuralTensor(InInputTensorData.FindChecked(81), TArray<int64>({1,3,27})), FNeuralTensor(InInputTensorDataSmall.FindChecked(30), TArray<int64>({2,3,5})),
+			FNeuralTensor(ENeuralDataType::Float, {1,2,23})},
+		// X={5,3,20}, W={2,3,5}, B={2}, Y={5,2,16}
+		{FNeuralTensor(InInputTensorData.FindChecked(300), TArray<int64>({5,3,20})), FNeuralTensor(InInputTensorDataSmall.FindChecked(30), TArray<int64>({2,3,5})),
+			FNeuralTensor(ENeuralDataType::Float, {5,2,16}), FNeuralTensor(InInputTensorDataSmall.FindChecked(2))},
+		// ----------------------------------------------------------------- 2D convolutions -----------------------------------------------------------------
+		// X={1,1,1,1}, W={1,1,1,1}, B={}, Y={1,1,1,1} (replicates 1D example)
+		{FNeuralTensor(InInputTensorData.FindChecked(1), TArray<int64>({1,1,1,1})), FNeuralTensor(InInputTensorDataSmall.FindChecked(1), TArray<int64>({1,1,1,1})),
+			FNeuralTensor(ENeuralDataType::Float, {1,1,1,1})},
+		// X={1,1,1,9}, W={1,1,1,5}, B={}, Y={1,1,1,5} (replicates 1D example)
+		{FNeuralTensor(InInputTensorData.FindChecked(9), TArray<int64>({1,1,1,9})), FNeuralTensor(InInputTensorDataSmall.FindChecked(5), TArray<int64>({1,1,1,5})),
+			FNeuralTensor(ENeuralDataType::Float, {1,1,1,5})},
+		// X={5,3,1,20}, W={2,3,1,5}, B={}, Y={5,2,1,16} (replicates 1D example)
+		{FNeuralTensor(InInputTensorData.FindChecked(300), TArray<int64>({5,3,1,20})), FNeuralTensor(InInputTensorDataSmall.FindChecked(30), TArray<int64>({2,3,1,5})),
+			FNeuralTensor(ENeuralDataType::Float, {5,2,1,16})},
+		// X={1,1,1,1}, W={1,1,1,1}, B={1}, Y={1,1,1,1} (replicates 1D example)
+		{FNeuralTensor(InInputTensorData.FindChecked(1), TArray<int64>({1,1,1,1})), FNeuralTensor(InInputTensorDataSmall.FindChecked(1), TArray<int64>({1,1,1,1})),
+			FNeuralTensor(ENeuralDataType::Float, {1,1,1,1}), FNeuralTensor(InInputTensorDataSmall.FindChecked(1))},
+		// X={1,1,5,10}, W={1,1,3,3}, B={}, Y={1,1,3,8}
+		{FNeuralTensor(InInputTensorData.FindChecked(50), TArray<int64>({1,1,5,10})), FNeuralTensor(InInputTensorDataSmall.FindChecked(9), TArray<int64>({1,1,3,3})),
+			FNeuralTensor(ENeuralDataType::Float, {1,1,3,8})},
+		// X={1,1,5,10}, W={1,1,3,3}, B={1}, Y={1,1,3,8}
+		{FNeuralTensor(InInputTensorData.FindChecked(50), TArray<int64>({1,1,5,10})), FNeuralTensor(InInputTensorDataSmall.FindChecked(9), TArray<int64>({1,1,3,3})),
+			FNeuralTensor(ENeuralDataType::Float, {1,1,3,8}), FNeuralTensor(InInputTensorDataSmall.FindChecked(1))},
+		// X={2,1,1,1}, W={5,1,1,1}, B={}, Y={2,5,1,1}
+		{FNeuralTensor(InInputTensorData.FindChecked(2), TArray<int64>({2,1,1,1})), FNeuralTensor(InInputTensorDataSmall.FindChecked(5), TArray<int64>({5,1,1,1})),
+			FNeuralTensor(ENeuralDataType::Float, {2,5,1,1})},
+		// X={2,1,1,1}, W={5,1,1,1}, B={5}, Y={2,5,1,1}
+		{FNeuralTensor(InInputTensorData.FindChecked(2), TArray<int64>({2,1,1,1})), FNeuralTensor(InInputTensorDataSmall.FindChecked(5), TArray<int64>({5,1,1,1})),
+			FNeuralTensor(ENeuralDataType::Float, {2,5,1,1}), FNeuralTensor(InInputTensorDataSmall.FindChecked(5))},
+		// X={2,3,1,1}, W={5,3,1,1}, B={}, Y={2,5,1,1}
+		{FNeuralTensor(InInputTensorData.FindChecked(6), TArray<int64>({2,3,1,1})), FNeuralTensor(InInputTensorDataSmall.FindChecked(15), TArray<int64>({5,3,1,1})),
+			FNeuralTensor(ENeuralDataType::Float, {2,5,1,1})},
+		// X={2,3,1,1}, W={5,3,1,1}, B={5}, Y={2,5,1,1}
+		{FNeuralTensor(InInputTensorData.FindChecked(6), TArray<int64>({2,3,1,1})), FNeuralTensor(InInputTensorDataSmall.FindChecked(15), TArray<int64>({5,3,1,1})),
+			FNeuralTensor(ENeuralDataType::Float, {2,5,1,1}), FNeuralTensor(InInputTensorDataSmall.FindChecked(5))},
+		// ----------------------------------------------------------------- 3D convolutions -----------------------------------------------------------------
+		// X={1,1,1,1,9}, W={1,1,1,1,5}, B={1}, Y={1,1,1,1,5} (replicates 1D example)
+		{FNeuralTensor(InInputTensorData.FindChecked(9), TArray<int64>({1,1,1,1,9})), FNeuralTensor(InInputTensorDataSmall.FindChecked(5), TArray<int64>({1,1,1,1,5})),
+			FNeuralTensor(ENeuralDataType::Float, {1,1,1,1,5}), FNeuralTensor(InInputTensorDataSmall.FindChecked(1))},
+		// X={1,1,1,9,1}, W={1,1,1,5,1}, B={}, Y={1,1,1,5,1} (replicates 1D example)
+		{FNeuralTensor(InInputTensorData.FindChecked(9), TArray<int64>({1,1,1,9,1})), FNeuralTensor(InInputTensorDataSmall.FindChecked(5), TArray<int64>({1,1,1,5,1})),
+			FNeuralTensor(ENeuralDataType::Float, {1,1,1,5,1})},
+		// X={1,1,1,1,1}, W={1,1,1,1,1}, B={}, Y={1,1,1,1,1} (replicates 1D example)
+		{FNeuralTensor(InInputTensorData.FindChecked(1), TArray<int64>({1,1,1,1,1})), FNeuralTensor(InInputTensorDataSmall.FindChecked(1), TArray<int64>({1,1,1,1,1})),
+			FNeuralTensor(ENeuralDataType::Float, {1,1,1,1,1})},
+		// X={1,1,6,5,5}, W={1,1,3,3,3}, B={}, Y={1,1,4,3,3}
+		{FNeuralTensor(InInputTensorData.FindChecked(150), TArray<int64>({1,1,6,5,5})), FNeuralTensor(InInputTensorDataSmall.FindChecked(27), TArray<int64>({1,1,3,3,3})),
+			FNeuralTensor(ENeuralDataType::Float, {1,1,4,3,3})},
+		// X={2,2,3,5,5}, W={9,2,3,3,3}, B={9}, Y={2,9,1,3,3}
+		{FNeuralTensor(InInputTensorData.FindChecked(300), TArray<int64>({2,2,3,5,5})), FNeuralTensor(InInputTensorDataSmall.FindChecked(486), TArray<int64>({9,2,3,3,3})),
+			FNeuralTensor(ENeuralDataType::Float, {2,9,1,3,3}), FNeuralTensor(InInputTensorDataSmall.FindChecked(9))},
+		// ----------------------------------------------------------------- 4D convolutions -----------------------------------------------------------------
+		// X={1,1,1,1,1,1}, W={1,1,1,1,1,1}, B={1}, Y={1,1,1,1,1,1} (replicates 1D example)
+		{FNeuralTensor(InInputTensorData.FindChecked(1), TArray<int64>({1,1,1,1,1,1})), FNeuralTensor(InInputTensorDataSmall.FindChecked(1), TArray<int64>({1,1,1,1,1,1})),
+			FNeuralTensor(ENeuralDataType::Float, {1,1,1,1,1,1}), FNeuralTensor(InInputTensorDataSmall.FindChecked(1))},
+		// X={1,1,1,1,1,9}, W={1,1,1,1,1,5}, B={}, Y={1,1,1,1,1,5} (replicates 1D example)
+		{FNeuralTensor(InInputTensorData.FindChecked(9), TArray<int64>({1,1,1,1,1,9})), FNeuralTensor(InInputTensorDataSmall.FindChecked(5), TArray<int64>({1,1,1,1,1,5})),
+			FNeuralTensor(ENeuralDataType::Float, {1,1,1,1,1,5})},
+		// X={1,1,3,4,5,5}, W={1,1,3,3,3,3}, B={}, Y={1,1,1,2,3,3}
+		{FNeuralTensor(InInputTensorData.FindChecked(300), TArray<int64>({1,1,3,4,5,5})), FNeuralTensor(InInputTensorDataSmall.FindChecked(81), TArray<int64>({1,1,3,3,3,3})),
+			FNeuralTensor(ENeuralDataType::Float, {1,1,1,2,3,3})},
+		// X={3,3,4,5,3,5}, W={2,3,3,3,3,3}, B={2}, Y={3,2,2,3,1,3}
+		{FNeuralTensor(InInputTensorData.FindChecked(2700), TArray<int64>({3,3,4,5,3,5})), FNeuralTensor(InInputTensorDataSmall.FindChecked(486), TArray<int64>({2,3,3,3,3,3})),
+			FNeuralTensor(ENeuralDataType::Float, {3,2,2,3,1,3}), FNeuralTensor(InInputTensorDataSmall.FindChecked(2))},
+		// ----------------------------------------------------------------- 8D convolutions -----------------------------------------------------------------
+		// X={1,1, 1,1,1,1,1,1,1,1}, W={1,1, 1,1,1,1,1,1,1,1}, B={1}, Y={1,1, 1,1,1,1,1,1,1,1} (replicates 1D example)
+		{FNeuralTensor(InInputTensorData.FindChecked(1), TArray<int64>({1,1, 1,1,1,1,1,1,1,1})), FNeuralTensor(InInputTensorDataSmall.FindChecked(1), TArray<int64>({1,1, 1,1,1,1,1,1,1,1})),
+			FNeuralTensor(ENeuralDataType::Float, {1,1, 1,1,1,1,1,1,1,1}), FNeuralTensor(InInputTensorDataSmall.FindChecked(1))},
+		// X={1,1, 1,1,1,1,1,1,1,9}, W={1,1, 1,1,1,1,1,1,1,5}, B={}, Y={1,1, 1,1,1,1,1,1,1,5} (replicates 1D example)
+		{FNeuralTensor(InInputTensorData.FindChecked(9), TArray<int64>({1,1, 1,1,1,1,1,1,1,9})), FNeuralTensor(InInputTensorDataSmall.FindChecked(5), TArray<int64>({1,1, 1,1,1,1,1,1,1,5})),
+			FNeuralTensor(ENeuralDataType::Float, {1,1, 1,1,1,1,1,1,1,5})},
+		// X={1,1, 1,1,1,1,3,4,5,5}, W={1,1, 1,1,1,1,3,3,3,3}, B={}, Y={1,1, 1,1,1,1,1,2,3,3} (replicates 4D example)
+		{FNeuralTensor(InInputTensorData.FindChecked(300), TArray<int64>({1,1, 1,1,1,1,3,4,5,5})), FNeuralTensor(InInputTensorDataSmall.FindChecked(81), TArray<int64>({1,1, 1,1,1,1,3,3,3,3})),
+			FNeuralTensor(ENeuralDataType::Float, {1,1, 1,1,1,1,1,2,3,3})},
+		// X={3,3, 1,1,1,1,4,5,3,5}, W={2,3, 1,1,1,1,3,3,3,3}, B={2}, Y={3,2, 1,1,1,1,2,3,1,3} (replicates 4D example)
+		{FNeuralTensor(InInputTensorData.FindChecked(2700), TArray<int64>({3,3, 1,1,1,1,4,5,3,5})), FNeuralTensor(InInputTensorDataSmall.FindChecked(486), TArray<int64>({2,3, 1,1,1,1,3,3,3,3})),
+			FNeuralTensor(ENeuralDataType::Float, {3,2, 1,1,1,1,2,3,1,3}), FNeuralTensor(InInputTensorDataSmall.FindChecked(2))},
+		// ----------------------------------------------------------------- Grouped convolutions -----------------------------------------------------------------
+		// X={1,4,1}, W={8,4/2,1}, B={}, Y={1,8,1}, Group = 2
+		{FNeuralTensor(TArray<float>({1,1,1,1}), TArray<int64>({1,4,1})),
+			FNeuralTensor(TArray<float>({1,1,1,1,1,1,1,1, 1,1,1,1,1,1,1,1, 1,1,1,1,1,1,1,1, 1,1,1,1,1,1,1,1}), TArray<int64>({8,4,1})),
+			FNeuralTensor(ENeuralDataType::Float, {1,8,1})},
+		{FNeuralTensor(TArray<float>({1,1,1,1}), TArray<int64>({1,4,1})),
+			FNeuralTensor(TArray<float>({1,1,1,1,1,1,1,1, 1,1,1,1,1,1,1,1}), TArray<int64>({8,2,1})),
+			FNeuralTensor(ENeuralDataType::Float, {1,8,1}), FNeuralTensor(), /*Group*/FNeuralTensor(TArray<float>({2}))},
+		// X={1,4,1}, W={8,4/2,1}, B={}, Y={1,8,1}, Group = 2
+		{FNeuralTensor(InInputTensorData.FindChecked(4), TArray<int64>({1,4,1})), FNeuralTensor(InInputTensorDataSmall.FindChecked(16), TArray<int64>({8,2,1})),
+			FNeuralTensor(ENeuralDataType::Float, {1,8,1}), FNeuralTensor(), /*Group*/FNeuralTensor(TArray<float>({2}))},
+		// X={1,9,9,6}, W={6,9/3,9,3}, B={}, Y={1,6,1,4}, Group = 3
+		{FNeuralTensor(InInputTensorData.FindChecked(486), TArray<int64>({1,9,9,6})), FNeuralTensor(InInputTensorDataSmall.FindChecked(486), TArray<int64>({6,3,9,3})),
+			FNeuralTensor(ENeuralDataType::Float, {1,6,1,4}), FNeuralTensor(), /*Group*/FNeuralTensor(TArray<float>({3}))},
+		// X={1,8,1,1,1}, W={4,8/4,1,1,1}, B={4}, Y={1,4,1,1,1}, Group = 4
+		{FNeuralTensor(InInputTensorData.FindChecked(8), TArray<int64>({1,8,1,1,1})), FNeuralTensor(InInputTensorDataSmall.FindChecked(8), TArray<int64>({4,2,1,1,1})),
+			FNeuralTensor(ENeuralDataType::Float, {1,4,1,1,1}), FNeuralTensor(InInputTensorDataSmall.FindChecked(4)), /*Group*/FNeuralTensor(TArray<float>({4}))},
+		// ----------------------------------------------------------------- AutoPad == SameLower/Upper -----------------------------------------------------------------
+		// X={1,1,9}, W={1,1,4}, B={}, Y={1,1,9}, AutoPad = EAutoPad::SameLower || EAutoPad::SameUpper
+		{FNeuralTensor(InInputTensorData.FindChecked(9), TArray<int64>({1,1,9})), FNeuralTensor(InInputTensorDataSmall.FindChecked(4), TArray<int64>({1,1,4})),
+			FNeuralTensor(ENeuralDataType::Float, {1,1,9}), FNeuralTensor(), /*Group*/FNeuralTensor(), /*AutoPad*/FNeuralTensor(TArray<float>({float(FConvBaseOperator::EAutoPad::SameLower)}))},
+		//                                        AutoPad = EAutoPad::SameLower || EAutoPad::SameUpper
+		{FNeuralTensor(InInputTensorData.FindChecked(9), TArray<int64>({1,1,9})), FNeuralTensor(InInputTensorDataSmall.FindChecked(4), TArray<int64>({1,1,4})),
+			FNeuralTensor(ENeuralDataType::Float, {1,1,9}), FNeuralTensor(), /*Group*/FNeuralTensor(), /*AutoPad*/FNeuralTensor(TArray<float>({float(FConvBaseOperator::EAutoPad::SameUpper)}))},
+		// X={1,1,9}, W={1,1,9}, B={}, Y={1,1,9}, AutoPad = EAutoPad::SameLower || EAutoPad::SameUpper
+		{FNeuralTensor(InInputTensorData.FindChecked(9), TArray<int64>({1,1,9})), FNeuralTensor(InInputTensorDataSmall.FindChecked(9), TArray<int64>({1,1,9})),
+			FNeuralTensor(ENeuralDataType::Float, {1,1,9}), FNeuralTensor(), /*Group*/FNeuralTensor(), /*AutoPad*/FNeuralTensor(TArray<float>({float(FConvBaseOperator::EAutoPad::SameLower)}))},
+		//                                        AutoPad = EAutoPad::SameLower || EAutoPad::SameUpper
+		{FNeuralTensor(InInputTensorData.FindChecked(9), TArray<int64>({1,1,9})), FNeuralTensor(InInputTensorDataSmall.FindChecked(9), TArray<int64>({1,1,9})),
+			FNeuralTensor(ENeuralDataType::Float, {1,1,9}), FNeuralTensor(), /*Group*/FNeuralTensor(), /*AutoPad*/FNeuralTensor(TArray<float>({float(FConvBaseOperator::EAutoPad::SameUpper)}))},
+		// X={1,1,9,9}, W={1,1,3,3}, B={}, Y={1,1,9,9}, AutoPad = EAutoPad::SameLower || EAutoPad::SameUpper
+		{FNeuralTensor(InInputTensorData.FindChecked(81), TArray<int64>({1,1,9,9})), FNeuralTensor(InInputTensorDataSmall.FindChecked(9), TArray<int64>({1,1,3,3})),
+			FNeuralTensor(ENeuralDataType::Float, {1,1,9,9}), FNeuralTensor(), /*Group*/FNeuralTensor(), /*AutoPad*/FNeuralTensor(TArray<float>({float(FConvBaseOperator::EAutoPad::SameLower)}))},
+		//                                              AutoPad = EAutoPad::SameLower || EAutoPad::SameUpper
+		{FNeuralTensor(InInputTensorData.FindChecked(81), TArray<int64>({1,1,9,9})), FNeuralTensor(InInputTensorDataSmall.FindChecked(9), TArray<int64>({1,1,3,3})),
+			FNeuralTensor(ENeuralDataType::Float, {1,1,9,9}), FNeuralTensor(), /*Group*/FNeuralTensor(), /*AutoPad*/FNeuralTensor(TArray<float>({float(FConvBaseOperator::EAutoPad::SameUpper)}))},
+		// X={1,1,9,9}, W={1,1,5,5}, B={}, Y={1,1,9,9}, AutoPad = EAutoPad::SameLower || EAutoPad::SameUpper
+		{FNeuralTensor(InInputTensorData.FindChecked(81), TArray<int64>({1,1,9,9})), FNeuralTensor(InInputTensorDataSmall.FindChecked(25), TArray<int64>({1,1,5,5})),
+			FNeuralTensor(ENeuralDataType::Float, {1,1,9,9}), FNeuralTensor(), /*Group*/FNeuralTensor(), /*AutoPad*/FNeuralTensor(TArray<float>({float(FConvBaseOperator::EAutoPad::SameLower)}))},
+		//                                              AutoPad = EAutoPad::SameLower || EAutoPad::SameUpper
+		{FNeuralTensor(InInputTensorData.FindChecked(81), TArray<int64>({1,1,9,9})), FNeuralTensor(InInputTensorDataSmall.FindChecked(25), TArray<int64>({1,1,5,5})),
+			FNeuralTensor(ENeuralDataType::Float, {1,1,9,9}), FNeuralTensor(), /*Group*/FNeuralTensor(), /*AutoPad*/FNeuralTensor(TArray<float>({float(FConvBaseOperator::EAutoPad::SameUpper)}))},
+		// X={1,1,5,5,4}, W={1,1,3,3,3}, B={}, Y={1,1,5,5,4}, AutoPad = EAutoPad::SameLower || EAutoPad::SameUpper
+		{FNeuralTensor(InInputTensorData.FindChecked(100), TArray<int64>({1,1,5,5,4})), FNeuralTensor(InInputTensorDataSmall.FindChecked(27), TArray<int64>({1,1,3,3,3})),
+			FNeuralTensor(ENeuralDataType::Float, {1,1,5,5,4}), FNeuralTensor(), /*Group*/FNeuralTensor(), /*AutoPad*/FNeuralTensor(TArray<float>({float(FConvBaseOperator::EAutoPad::SameLower)}))},
+		//                                                    AutoPad = EAutoPad::SameLower || EAutoPad::SameUpper
+		{FNeuralTensor(InInputTensorData.FindChecked(100), TArray<int64>({1,1,5,5,4})), FNeuralTensor(InInputTensorDataSmall.FindChecked(27), TArray<int64>({1,1,3,3,3})),
+			FNeuralTensor(ENeuralDataType::Float, {1,1,5,5,4}), FNeuralTensor(), /*Group*/FNeuralTensor(), /*AutoPad*/FNeuralTensor(TArray<float>({float(FConvBaseOperator::EAutoPad::SameUpper)}))},
+		// X={1,1,1,1,5,5,4}, W={1,1,1,1,3,3,3}, B={}, Y={1,1,1,1,5,5,4}, AutoPad = EAutoPad::SameLower || EAutoPad::SameUpper
+		{FNeuralTensor(InInputTensorData.FindChecked(100), TArray<int64>({1,1,1,1,5,5,4})), FNeuralTensor(InInputTensorDataSmall.FindChecked(27), TArray<int64>({1,1,1,1,3,3,3})),
+			FNeuralTensor(ENeuralDataType::Float, {1,1,1,1,5,5,4}), FNeuralTensor(), /*Group*/FNeuralTensor(), /*AutoPad*/FNeuralTensor(TArray<float>({float(FConvBaseOperator::EAutoPad::SameLower)}))},
+		//                                                    AutoPad = EAutoPad::SameLower || EAutoPad::SameUpper
+		{FNeuralTensor(InInputTensorData.FindChecked(100), TArray<int64>({1,1,5,5,4,1,1})), FNeuralTensor(InInputTensorDataSmall.FindChecked(27), TArray<int64>({1,1,3,3,3,1,1})),
+			FNeuralTensor(ENeuralDataType::Float, {1,1,5,5,4,1,1}), FNeuralTensor(), /*Group*/FNeuralTensor(), /*AutoPad*/FNeuralTensor(TArray<float>({float(FConvBaseOperator::EAutoPad::SameUpper)}))},
+		// X={2,3,5,2,3,3,5}, W={2,3,5,2,1,1,5}, B={}, Y={2,2,5,2,3,3,5}, AutoPad = EAutoPad::SameLower || EAutoPad::SameUpper
+		{FNeuralTensor(InInputTensorData.FindChecked(2700), TArray<int64>({2,3,5,2,3,3,5})), FNeuralTensor(InInputTensorDataSmall.FindChecked(300), TArray<int64>({2,3,5,2,1,1,5})),
+			FNeuralTensor(ENeuralDataType::Float, {2,2,5,2,3,3,5}), FNeuralTensor(), /*Group*/FNeuralTensor(), /*AutoPad*/FNeuralTensor(TArray<float>({float(FConvBaseOperator::EAutoPad::SameLower)}))},
+		//																  AutoPad = EAutoPad::SameLower || EAutoPad::SameUpper
+		{FNeuralTensor(InInputTensorData.FindChecked(2700), TArray<int64>({2,3,5,2,3,3,5})), FNeuralTensor(InInputTensorDataSmall.FindChecked(300), TArray<int64>({2,3,5,2,1,1,5})),
+			FNeuralTensor(ENeuralDataType::Float, {2,2,5,2,3,3,5}), FNeuralTensor(), /*Group*/FNeuralTensor(), /*AutoPad*/FNeuralTensor(TArray<float>({float(FConvBaseOperator::EAutoPad::SameUpper)}))},
+		// ----------------------------------------------------------------- AutoPad = EAutoPad::NotSet (i.e., manual padding) -----------------------------------------------------------------
+		// X={1,1,9}, W={1,1,4}, B={}, Y={1,1,16}
+		{FNeuralTensor(InInputTensorData.FindChecked(9), TArray<int64>({1,1,9})), FNeuralTensor(InInputTensorDataSmall.FindChecked(4), TArray<int64>({1,1,4})),
+			FNeuralTensor(ENeuralDataType::Float, {1,1,16}), FNeuralTensor(), /*Group*/FNeuralTensor(), /*AutoPad*/FNeuralTensor(), /*Pads*/FNeuralTensor(TArray<float>({0,10}))},
+		// X={1,1,9,9}, W={1,1,3,3}, B={}, Y={1,1,12,17}
+		{FNeuralTensor(InInputTensorData.FindChecked(81), TArray<int64>({1,1,9,9})), FNeuralTensor(InInputTensorDataSmall.FindChecked(9), TArray<int64>({1,1,3,3})),
+			FNeuralTensor(ENeuralDataType::Float, {1,1,12,17}), FNeuralTensor(), /*Group*/FNeuralTensor(), /*AutoPad*/FNeuralTensor(), /*Pads*/FNeuralTensor(TArray<float>({5,0, 5,5}))},
+		// X={1,1,5,5,4}, W={1,1,3,3,3}, B={}, Y={1,1,4,8,11}
+		{FNeuralTensor(InInputTensorData.FindChecked(100), TArray<int64>({1,1,5,5,4})), FNeuralTensor(InInputTensorDataSmall.FindChecked(27), TArray<int64>({1,1,3,3,3})),
+			FNeuralTensor(ENeuralDataType::Float, {1,1,4,8,11}), FNeuralTensor(), /*Group*/FNeuralTensor(), /*AutoPad*/FNeuralTensor(), /*Pads*/FNeuralTensor(TArray<float>({0,1, 2,3, 4,5}))},
+		// X={1,1,1,1,5,5,4}, W={1,1,1,1,3,3,3}, B={}, Y={1,1,2,6,12,10,5}
+		{FNeuralTensor(InInputTensorData.FindChecked(100), TArray<int64>({1,1,1,1,5,5,4})), FNeuralTensor(InInputTensorDataSmall.FindChecked(27), TArray<int64>({1,1,1,1,3,3,3})),
+			FNeuralTensor(ENeuralDataType::Float, {1,1,2,6,12,10,5}), FNeuralTensor(), /*Group*/FNeuralTensor(), /*AutoPad*/FNeuralTensor(), /*Pads*/FNeuralTensor(TArray<float>({0,1, 2,3, 4,5, 4,3, 2,1}))},
+		// ----------------------------------------------------------------- Dilations != 1 -----------------------------------------------------------------
+		// X={1,1,9}, W={1,1,4}, B={}, Y={1,1,3}
+		{FNeuralTensor(InInputTensorData.FindChecked(9), TArray<int64>({1,1,9})), FNeuralTensor(InInputTensorDataSmall.FindChecked(4), TArray<int64>({1,1,4})),
+			FNeuralTensor(ENeuralDataType::Float, {1,1,3}), FNeuralTensor(), /*Group*/FNeuralTensor(), /*AutoPad*/FNeuralTensor(), /*Pads*/FNeuralTensor(), /*Dilations*/FNeuralTensor(TArray<float>({2}))},
+		// X={1,1,9,9}, W={1,1,3,3}, B={}, Y={1,1,3,5}
+		{FNeuralTensor(InInputTensorData.FindChecked(81), TArray<int64>({1,1,9,9})), FNeuralTensor(InInputTensorDataSmall.FindChecked(9), TArray<int64>({1,1,3,3})),
+			FNeuralTensor(ENeuralDataType::Float, {1,1,3,5}), FNeuralTensor(), /*Group*/FNeuralTensor(), /*AutoPad*/FNeuralTensor(), /*Pads*/FNeuralTensor(), /*Dilations*/FNeuralTensor(TArray<float>({3,2}))},
+		// X={1,1,5,4,5}, W={1,1,3,3,3}, B={}, Y={1,1,1,2,1}
+		{FNeuralTensor(InInputTensorData.FindChecked(100), TArray<int64>({1,1,5,4,5})), FNeuralTensor(InInputTensorDataSmall.FindChecked(27), TArray<int64>({1,1,3,3,3})),
+			FNeuralTensor(ENeuralDataType::Float, {1,1,1,2,1}), FNeuralTensor(), /*Group*/FNeuralTensor(), /*AutoPad*/FNeuralTensor(), /*Pads*/FNeuralTensor(), /*Dilations*/FNeuralTensor(TArray<float>({2,1,2}))},
+		// X={1,1,1,1,5,4,5}, W={1,1,1,1,3,3,3}, B={}, Y={1,1,1,1,1,2,1}
+		{FNeuralTensor(InInputTensorData.FindChecked(100), TArray<int64>({1,1,1,1,5,4,5})), FNeuralTensor(InInputTensorDataSmall.FindChecked(27), TArray<int64>({1,1,1,1,3,3,3})),
+			FNeuralTensor(ENeuralDataType::Float, {1,1,1,1,1,2,1}), FNeuralTensor(), /*Group*/FNeuralTensor(), /*AutoPad*/FNeuralTensor(), /*Pads*/FNeuralTensor(), /*Dilations*/FNeuralTensor(TArray<float>({99,10,2,1,2}))},
+		// ----------------------------------------------------------------- Strides != 1 -----------------------------------------------------------------
+		// Required manual padding = 0 to avoid this unit tester to accidentally auto-trigger AutoPad == SameLower/Upper
+		// X={1,1,9}, W={1,1,4}, B={}, Y={1,1,2}
+		{FNeuralTensor(InInputTensorData.FindChecked(9), TArray<int64>({1,1,9})), FNeuralTensor(InInputTensorDataSmall.FindChecked(4), TArray<int64>({1,1,4})),
+			FNeuralTensor(ENeuralDataType::Float, {1,1,2}), FNeuralTensor(), /*Group*/FNeuralTensor(), /*AutoPad*/FNeuralTensor(), /*Pads*/FNeuralTensor(),
+			/*Dilations*/FNeuralTensor(), /*Strides*/FNeuralTensor(TArray<float>({5}))},
+		// X={1,1,9,9}, W={1,1,3,3}, B={}, Y={1,1,1,3}
+		{FNeuralTensor(InInputTensorData.FindChecked(81), TArray<int64>({1,1,9,9})), FNeuralTensor(InInputTensorDataSmall.FindChecked(9), TArray<int64>({1,1,3,3})),
+			FNeuralTensor(ENeuralDataType::Float, {1,1,1,3}), FNeuralTensor(), /*Group*/FNeuralTensor(), /*AutoPad*/FNeuralTensor(), /*Pads*/FNeuralTensor(),
+			/*Dilations*/FNeuralTensor(), /*Strides*/FNeuralTensor(TArray<float>({99,3}))},
+		// X={1,1,5,4,5}, W={1,1,3,3,3}, B={}, Y={1,1,2,2,1}
+		{FNeuralTensor(InInputTensorData.FindChecked(100), TArray<int64>({1,1,5,4,5})), FNeuralTensor(InInputTensorDataSmall.FindChecked(27), TArray<int64>({1,1,3,3,3})),
+			FNeuralTensor(ENeuralDataType::Float, {1,1,2,2,1}), FNeuralTensor(), /*Group*/FNeuralTensor(), /*AutoPad*/FNeuralTensor(), /*Pads*/FNeuralTensor(),
+			/*Dilations*/FNeuralTensor(), /*Strides*/FNeuralTensor(TArray<float>({2,1,99}))},
+		// X={1,1,1,1,5,4,5}, W={1,1,1,1,3,3,3}, B={}, Y={1,1,1,1,2,2,1}
+		{FNeuralTensor(InInputTensorData.FindChecked(100), TArray<int64>({1,1,1,1,5,4,5})), FNeuralTensor(InInputTensorDataSmall.FindChecked(27), TArray<int64>({1,1,1,1,3,3,3})),
+			FNeuralTensor(ENeuralDataType::Float, {1,1,1,1,2,2,1}), FNeuralTensor(), /*Group*/FNeuralTensor(), /*AutoPad*/FNeuralTensor(), /*Pads*/FNeuralTensor(),
+			/*Dilations*/FNeuralTensor(), /*Strides*/FNeuralTensor(TArray<float>({1,50,2,1,99}))},
+		// ----------------------------------------------------------------- Strides != 1, Multi-channel -----------------------------------------------------------------
+		// X={5,5,25}, W={2,5,9}, B={}, Y={5,2,9}
+		{FNeuralTensor(InInputTensorData.FindChecked(625), TArray<int64>({5,5,25})), FNeuralTensor(InInputTensorDataSmall.FindChecked(90), TArray<int64>({2,5,9})),
+			FNeuralTensor(ENeuralDataType::Float, {5,2,9}), FNeuralTensor(), /*Group*/FNeuralTensor(), /*AutoPad*/FNeuralTensor(), /*Pads*/FNeuralTensor(),
+			/*Dilations*/FNeuralTensor(), /*Strides*/FNeuralTensor(TArray<float>({2}))},
+		// X={5,5,5,5}, W={2,5,3,3}, B={}, Y={5,2,2,2}
+		{FNeuralTensor(InInputTensorData.FindChecked(625), TArray<int64>({5,5,5,5})), FNeuralTensor(InInputTensorDataSmall.FindChecked(90), TArray<int64>({2,5,3,3})),
+			FNeuralTensor(ENeuralDataType::Float, {5,2,2,2}), FNeuralTensor(), /*Group*/FNeuralTensor(), /*AutoPad*/FNeuralTensor(), /*Pads*/FNeuralTensor(),
+			/*Dilations*/FNeuralTensor(), /*Strides*/FNeuralTensor(TArray<float>({2,2}))},
+		// X={5,5,5,5,1}, W={2,5,3,3,1}, B={}, Y={5,2,2,2,1}
+		{FNeuralTensor(InInputTensorData.FindChecked(625), TArray<int64>({5,5,5,5,1})), FNeuralTensor(InInputTensorDataSmall.FindChecked(90), TArray<int64>({2,5,3,3,1})),
+			FNeuralTensor(ENeuralDataType::Float, {5,2,2,2,1}), FNeuralTensor(), /*Group*/FNeuralTensor(), /*AutoPad*/FNeuralTensor(), /*Pads*/FNeuralTensor(),
+			/*Dilations*/FNeuralTensor(), /*Strides*/FNeuralTensor(TArray<float>({2,2,1}))},
+		// X={5,5,5,5,1,1,1}, W={2,5,3,3,1,1,1}, B={}, Y={5,2,2,2,1,1,1}
+		{FNeuralTensor(InInputTensorData.FindChecked(625), TArray<int64>({5,5,5,5,1,1,1})), FNeuralTensor(InInputTensorDataSmall.FindChecked(90), TArray<int64>({2,5,3,3,1,1,1})),
+			FNeuralTensor(ENeuralDataType::Float, {5,2,2,2,1,1,1}), FNeuralTensor(), /*Group*/FNeuralTensor(), /*AutoPad*/FNeuralTensor(), /*Pads*/FNeuralTensor(),
+			/*Dilations*/FNeuralTensor(), /*Strides*/FNeuralTensor(TArray<float>({2,2,1,1,1}))},
+		// ----------------------------------------------------------------- Manual padding or AutoPad == SameLower/Upper, Dilations != 1, Strides != 1 -----------------------------------------------------------------
+		// X={1,1,18}, W={1,1,4}, B={}, Y={1,1,9}, AutoPad = EAutoPad::SameLower || EAutoPad::SameUpper
+		{FNeuralTensor(InInputTensorData.FindChecked(18), TArray<int64>({1,1,18})), FNeuralTensor(InInputTensorDataSmall.FindChecked(4), TArray<int64>({1,1,4})),
+			FNeuralTensor(ENeuralDataType::Float, {1,1,9}), FNeuralTensor(), /*Group*/FNeuralTensor(), /*AutoPad*/FNeuralTensor(TArray<float>({float(FConvBaseOperator::EAutoPad::SameLower)})), /*Pads*/FNeuralTensor(),
+			/*Dilations*/FNeuralTensor(TArray<float>({2})), /*Strides*/FNeuralTensor(TArray<float>({2}))},
+		//										   AutoPad = EAutoPad::SameUpper || EAutoPad::SameLower
+		{FNeuralTensor(InInputTensorData.FindChecked(18), TArray<int64>({1,1,18})), FNeuralTensor(InInputTensorDataSmall.FindChecked(4), TArray<int64>({1,1,4})),
+			FNeuralTensor(ENeuralDataType::Float, {1,1,9}), FNeuralTensor(), /*Group*/FNeuralTensor(), /*AutoPad*/FNeuralTensor(TArray<float>({float(FConvBaseOperator::EAutoPad::SameUpper)})), /*Pads*/FNeuralTensor(),
+			/*Dilations*/FNeuralTensor(TArray<float>({2})), /*Strides*/FNeuralTensor(TArray<float>({2}))},
+		// X={1,1,100}, W={1,1,3}, B={}, Y={1,1,21}
+		{FNeuralTensor(InInputTensorData.FindChecked(100), TArray<int64>({1,1,100})), FNeuralTensor(InInputTensorDataSmall.FindChecked(3), TArray<int64>({1,1,3})),
+			FNeuralTensor(ENeuralDataType::Float, {1,1,21}), FNeuralTensor(), /*Group*/FNeuralTensor(), /*AutoPad*/FNeuralTensor(), /*Pads*/FNeuralTensor(TArray<float>({10,5})),
+			/*Dilations*/FNeuralTensor(TArray<float>({5})), /*Strides*/FNeuralTensor(TArray<float>({5}))},
+		// X={1,1,9,9}, W={1,1,3,3}, B={}, Y={1,1,16,14}
+		{FNeuralTensor(InInputTensorData.FindChecked(81), TArray<int64>({1,1,9,9})), FNeuralTensor(InInputTensorDataSmall.FindChecked(9), TArray<int64>({1,1,3,3})),
+			FNeuralTensor(ENeuralDataType::Float, {1,1,16,14}), FNeuralTensor(), /*Group*/FNeuralTensor(), /*AutoPad*/FNeuralTensor(), /*Pads*/FNeuralTensor(TArray<float>({10,1,0,7})),
+			/*Dilations*/FNeuralTensor(TArray<float>({2,1})), /*Strides*/FNeuralTensor(TArray<float>({1,1}))},
+		// X={1,1,9,9}, W={1,1,3,3}, B={}, Y={1,1,4,5}
+		{FNeuralTensor(InInputTensorData.FindChecked(81), TArray<int64>({1,1,9,9})), FNeuralTensor(InInputTensorDataSmall.FindChecked(9), TArray<int64>({1,1,3,3})),
+			FNeuralTensor(ENeuralDataType::Float, {1,1,4,5}), FNeuralTensor(), /*Group*/FNeuralTensor(), /*AutoPad*/FNeuralTensor(), /*Pads*/FNeuralTensor(TArray<float>({10,1,0,7})),
+			/*Dilations*/FNeuralTensor(TArray<float>({2,1})), /*Strides*/FNeuralTensor(TArray<float>({4,3}))},
+		// X={1,1,5,4,5}, W={1,1,3,3,3}, B={}, Y={1,1,5,4,5}
+		{FNeuralTensor(InInputTensorData.FindChecked(100), TArray<int64>({1,1,5,4,5})), FNeuralTensor(InInputTensorDataSmall.FindChecked(27), TArray<int64>({1,1,3,3,3})),
+			FNeuralTensor(ENeuralDataType::Float, {1,1,5,4,5}), FNeuralTensor(), /*Group*/FNeuralTensor(), /*AutoPad*/FNeuralTensor(), /*Pads*/FNeuralTensor(TArray<float>({10,6,1,12,15,0})),
+			/*Dilations*/FNeuralTensor(TArray<float>({3,3,3})), /*Strides*/FNeuralTensor(TArray<float>({3,3,3}))},
+		// X={1,1,1,1,5,4,5}, W={1,1,1,1,3,3,3}, B={}, Y={1,1,1,1,5,4,5}
+		{FNeuralTensor(InInputTensorData.FindChecked(100), TArray<int64>({1,1,1,1,5,4,5})), FNeuralTensor(InInputTensorDataSmall.FindChecked(27), TArray<int64>({1,1,1,1,3,3,3})),
+			FNeuralTensor(ENeuralDataType::Float, {1,1,1,1,5,4,5}), FNeuralTensor(), /*Group*/FNeuralTensor(), /*AutoPad*/FNeuralTensor(), /*Pads*/FNeuralTensor(TArray<float>({0,0,0,0,10,6,1,12,15,0})),
+			/*Dilations*/FNeuralTensor(TArray<float>({8,5,3,3,3})), /*Strides*/FNeuralTensor(TArray<float>({2,9,3,3,3}))}
+		});
+	TArray<TArray<FNeuralTensor*>> InputTensors;
+	TArray<TArray<FNeuralTensor*>> OutputTensors;
+	TArray<TSharedPtr<FNeuralOperator>> Operators;
+	for (int32 TensorIndex = 0; TensorIndex < Tensors.Num(); ++TensorIndex)
+	{
+		TArray<FNeuralTensor>& TensorsIndex = Tensors[TensorIndex];
+		if (TensorsIndex.Num() > 3 && TensorsIndex[3].Num() > 0)
+		{
+			InputTensors.Push({ &TensorsIndex[0], &TensorsIndex[1], &TensorsIndex[3] });
+		}
+		else
+		{
+			InputTensors.Push({ &TensorsIndex[0], &TensorsIndex[1] });
+		}
+		OutputTensors.Push({ &TensorsIndex[2] });
+		// Set Group - Hacky but simple way of getting group value
+		//const int64 Group = (TensorsIndex.Num() > 4 && TensorsIndex[4].Num() > 0 ? FMath::RoundToInt(TensorsIndex[4].At<float>(0)) : 1);
+		const TArray<int64> GroupArray = FloatTensorToIntArray(TensorsIndex, 4);
+		if (GroupArray.Num() > 1)
+		{
+			ensureMsgf(false, TEXT("GroupArray.Num() = %d == 0 | 1 failed."), GroupArray.Num());
+			return;
+		}
+		const int64 Group = (GroupArray.Num() > 0 ? GroupArray[0] : 1);
+		// Set FConvBaseOperator::EAutoPad - Hacky but simple way of getting this TArray
+		FConvBaseOperator::EAutoPad AutoPad = FConvBaseOperator::EAutoPad::NotSet;
+		const TArray<int64> AutoPadArray = FloatTensorToIntArray(TensorsIndex, 5);
+		if (AutoPadArray.Num() > 0)
+		{
+			if (AutoPadArray.Num() != 1)
+			{
+				ensureMsgf(false, TEXT("AutoPadArray should always be 1."));
+				return;
+			}
+			AutoPad = FConvBaseOperator::EAutoPad(AutoPadArray[0]);
+		}
+		// Set Pads - Hacky but simple way of getting this TArray
+		const TArray<int64> Pads = FloatTensorToIntArray(TensorsIndex, 6);
+		// Set Dilations - Hacky but simple way of getting this TArray
+		const TArray<int64> Dilations = FloatTensorToIntArray(TensorsIndex, 7);
+		// Set Strides - Hacky but simple way of getting this TArray
+		const TArray<int64> Strides = FloatTensorToIntArray(TensorsIndex, 8);
+		// Set FConvBaseOperator::EAutoPad - Hacky but simple way of getting this TArray
+		// Create and push FConvOperator
+		Operators.Push(MakeShared<FConvOperator>(AutoPad, Dilations, Group, /*KernelShape*/TArray<int64>(), Pads, Strides));
+		// Remove all non-required/auxiliary tensors from InputTensors to keep it in X, W, and B (if B exists). Otherwise, those would be used by the network and moved into GPU memory
+		while (Tensors[TensorIndex].Num() > 4 || Tensors[TensorIndex].Last().Num() < 1)
+		{
+			Tensors[TensorIndex].Pop();
+		}
+		// TensorsIndex is no longer valid after this while loop!
+	}
+	TestOperator(InOutNetworkInferenceQAAsset, Tensors, InputTensors, OutputTensors, Operators, TEXT("FConvOperator"), InGroundTruthDirectory, InZeroThreshold,
+		/*bShouldOperatorsProvideSameResult*/false, /*bInShouldRunGPU*/true);
+}
+
+void FOperatorTester::TestConvTransposeOperator(UNeuralNetworkInferenceQAAsset* InOutNetworkInferenceQAAsset, const TMap<int32, TArray<float>>& InInputTensorData, const TMap<int32, TArray<float>>& InInputTensorDataSmall,
+	const float InZeroThreshold, const FString& InGroundTruthDirectory, const int32 InOperatorCounter, const int32 InOperatorTotal)
+{
+	UE_LOG(LogNeuralNetworkInferenceQA, Display, TEXT("--------------- %d/%d FConvTransposeOperator"), InOperatorCounter, InOperatorTotal);
+	// Inout tensors
+	TArray<TArray<FNeuralTensor>> Tensors({
+		// ----------------------------------------------------------------- ONNX - [REDACTED_PROFILE_URL]/onnx/blob/master/docs/TestCoverage.md#ConvTranspose -----------------------------------------------------------------
+		// convtranspose_1d
+		{FNeuralTensor(TArray<float>({0,1,2}), TArray<int64>({1,1,3})),
+			FNeuralTensor(TArray<float>({1,1,1, 1,1,1}), TArray<int64>({1,2,3})),
+			FNeuralTensor(ENeuralDataType::Float, {1,2,5})}, // [[[0.0 1.0 3.0 3.0 2.0], [0.0 1.0 3.0 3.0 2.0]]]
+		// convtranspose
+		{FNeuralTensor(TArray<float>({0,1,2, 3,4,5, 6,7,8}), TArray<int64>({1,1,3,3})),
+			FNeuralTensor(TArray<float>({1,1,1, 1,1,1, 1,1,1,   1,1,1, 1,1,1, 1,1,1}), TArray<int64>({1,2,3,3})),
+			FNeuralTensor(ENeuralDataType::Float, {1,2,5,5})}, // [[[[0.0 1.0 3.0 3.0 2.0], [3.0 8.0 15.0 12.0 7.0], [9.0 21.0 36.0 27.0 15.0], [9.0 20.0 33.0 24.0 13.0], [6.0 13.0 21.0 15.0 8.0]], [repeated]]
+		// convtranspose_3d
+		{FNeuralTensor(TArray<float>({0,1,2,3,4, 5,6,7,8,9, 10,11,12,13,14, 15,16,17,18,19, 20,21,22,23,24, 25,26,27,28,29, 30,31,32,33,34, 35,36,37,38,39, 40,41,42,43,44, 45,46,47,48,49, 50,51,52,53,54, 55,56,57,58,59}), TArray<int64>({1,1,3,4,5})),
+			FNeuralTensor(TArray<float>({1,1,1, 1,1,1, 1,1,1, 1,1,1, 1,1,1, 1,1,1, 1,1,1, 1,1,1, 1,1,1, 1,1,1, 1,1,1, 1,1,1, 1,1,1, 1,1,1, 1,1,1, 1,1,1, 1,1,1, 1,1,1}), TArray<int64>({1,2,3,3,3})),
+			FNeuralTensor(ENeuralDataType::Float, {1,2,5,6,7})},
+		// convtranspose_attributes
+		// Y = [[[[0.0 0.0 1.0 1.0 3.0 2.0 2.0 0.0], [0.0 0.0 1.0 1.0 3.0 2.0 2.0 0.0], [0.0 0.0 1.0 1.0 3.0 2.0 2.0 0.0], [3.0 3.0 7.0 4.0 9.0 5.0 5.0 0.0], [3.0 3.0 7.0 4.0 9.0 5.0 5.0 0.0], [3.0 3.0 7.0 4.0 9.0 5.0 5.0 0.0], [6.0 6.0 13.0 7.0 15.0 8.0 8.0 0.0],
+		//        [6.0 6.0 13.0 7.0 15.0 8.0 8.0 0.0], [6.0 6.0 13.0 7.0 15.0 8.0 8.0 0.0], [0.0 0.0 0.0 0.0 0.0 0.0 0.0 0.0]], [[repeated]]
+		{FNeuralTensor(TArray<float>({0,1,2, 3,4,5, 6,7,8}), TArray<int64>({1,1,3,3})),
+			FNeuralTensor(TArray<float>({1,1,1, 1,1,1, 1,1,1,   1,1,1, 1,1,1, 1,1,1}), TArray<int64>({1,2,3,3})),
+			FNeuralTensor(ENeuralDataType::Float, {1,2,10,8}), FNeuralTensor(), /*Group*/FNeuralTensor(), /*AutoPad*/FNeuralTensor(), /*Pads*/FNeuralTensor(),
+			/*Dilations*/FNeuralTensor(), /*Strides*/FNeuralTensor(TArray<float>({3,2})), /*OutputPadding*/FNeuralTensor(), /*OutputShape*/FNeuralTensor(TArray<float>({10,8}))},
+		{FNeuralTensor(TArray<float>({0,1,2, 3,4,5, 6,7,8}), TArray<int64>({1,1,3,3})),
+			FNeuralTensor(TArray<float>({1,1,1, 1,1,1, 1,1,1,   1,1,1, 1,1,1, 1,1,1}), TArray<int64>({1,2,3,3})),
+			FNeuralTensor(ENeuralDataType::Float, {1,2,10,8}), FNeuralTensor(), /*Group*/FNeuralTensor(), /*AutoPad*/FNeuralTensor(), /*Pads*/FNeuralTensor(),
+			/*Dilations*/FNeuralTensor(), /*Strides*/FNeuralTensor(TArray<float>({3,2})), /*OutputPadding*/FNeuralTensor(TArray<float>({1,1})), /*OutputShape*/FNeuralTensor()},
+		{FNeuralTensor(TArray<float>({0,1,2, 3,4,5, 6,7,8}), TArray<int64>({1,1,3,3})),
+			FNeuralTensor(TArray<float>({1,1,1, 1,1,1, 1,1,1,   1,1,1, 1,1,1, 1,1,1}), TArray<int64>({1,2,3,3})),
+			FNeuralTensor(ENeuralDataType::Float, {1,2,10,8}), FNeuralTensor(), /*Group*/FNeuralTensor(), /*AutoPad*/FNeuralTensor(), /*Pads*/FNeuralTensor(),
+			/*Dilations*/FNeuralTensor(), /*Strides*/FNeuralTensor(TArray<float>({3,2})), /*OutputPadding*/FNeuralTensor(TArray<float>({1,1})), /*OutputShape*/FNeuralTensor(TArray<float>({10,8}))},
+		// convtranspose_autopad_same
+		// Y = [[[[0.0 0.0 1.0 1.0 3.0 2.0], [0.0 0.0 1.0 1.0 3.0 2.0], [3.0 3.0 8.0 5.0 12.0 7.0], [3.0 3.0 7.0 4.0 9.0 5.0], [9.0 9.0 20.0 11.0 24.0 13.0], [6.0 6.0 13.0 7.0 15.0 8.0]],
+		//       [[0.0 0.0 1.0 1.0 3.0 2.0], [0.0 0.0 1.0 1.0 3.0 2.0], [3.0 3.0 8.0 5.0 12.0 7.0], [3.0 3.0 7.0 4.0 9.0 5.0], [9.0 9.0 20.0 11.0 24.0 13.0], [6.0 6.0 13.0 7.0 15.0 8.0]]]]
+		{FNeuralTensor(TArray<float>({0,1,2, 3,4,5, 6,7,8}), TArray<int64>({1,1,3,3})),
+			FNeuralTensor(TArray<float>({1,1,1, 1,1,1, 1,1,1,   1,1,1, 1,1,1, 1,1,1}), TArray<int64>({1,2,3,3})),
+			FNeuralTensor(ENeuralDataType::Float, {1,2,6,6}), FNeuralTensor(), /*Group*/FNeuralTensor(), /*AutoPad*/FNeuralTensor(TArray<float>({float(FConvBaseOperator::EAutoPad::SameLower)})), /*Pads*/FNeuralTensor(),
+			/*Dilations*/FNeuralTensor(), /*Strides*/FNeuralTensor(TArray<float>({2,2})), /*OutputPadding*/FNeuralTensor(), /*OutputShape*/FNeuralTensor()},
+		// convtranspose_dilations
+		// Y = [[[[21.0 56.0 13.0 16.0 2.0], [63.0 35.0 67.0 10.0 14.0], [24.0 22.0 76.0 76.0 21.0], [9.0 5.0 88.0 45.0 63.0], [3.0 2.0 33.0 18.0 54.0]]]]
+		{FNeuralTensor(TArray<float>({3,8,1, 9,5,7, 3,2,6}), TArray<int64>({1,1,3,3})),
+			FNeuralTensor(TArray<float>({7,2, 1,9}), TArray<int64>({1,1,2,2})),
+			FNeuralTensor(ENeuralDataType::Float, {1,1,5,5}), FNeuralTensor(), /*Group*/FNeuralTensor(), /*AutoPad*/FNeuralTensor(), /*Pads*/FNeuralTensor(),
+			/*Dilations*/FNeuralTensor(TArray<float>({2,2})), /*Strides*/FNeuralTensor(), /*OutputPadding*/FNeuralTensor(), /*OutputShape*/FNeuralTensor()},
+		// convtranspose_pads
+		// Output size (thus solution) in onnx website is wrong. Replaced Y size to [6, 4]
+		// Y = [[[[0.0 1.0 1.0 3.0], [0.0 1.0 1.0 3.0], [3.0 7.0 4.0 9.0], [3.0 7.0 4.0 9.0], [3.0 7.0 4.0 9.0], [6.0 13.0 7.0 15.0]], [[0.0 1.0 1.0 3.0], [0.0 1.0 1.0 3.0], [3.0 7.0 4.0 9.0], [3.0 7.0 4.0 9.0], [3.0 7.0 4.0 9.0], [6.0 13.0 7.0 15.0]]]]
+		{FNeuralTensor(TArray<float>({0,1,2, 3,4,5, 6,7,8}), TArray<int64>({1,1,3,3})),
+			FNeuralTensor(TArray<float>({1,1,1, 1,1,1, 1,1,1,   1,1,1, 1,1,1, 1,1,1}), TArray<int64>({1,2,3,3})),
+			FNeuralTensor(ENeuralDataType::Float, {1,2,6,4}), FNeuralTensor(), /*Group*/FNeuralTensor(), /*AutoPad*/FNeuralTensor(), /*Pads*/FNeuralTensor(TArray<float>({1,2,1,2})),
+			/*Dilations*/FNeuralTensor(), /*Strides*/FNeuralTensor(TArray<float>({3,2})), /*OutputPadding*/FNeuralTensor(), /*OutputShape*/FNeuralTensor()},
+		// ----------------------------------------------------------------- 1D convolutions -----------------------------------------------------------------
+		// X={1,1,1}, W={1,1,1}, B={1}, Y={1,1,1}
+		{FNeuralTensor(InInputTensorData.FindChecked(1), TArray<int64>({1,1,1})), FNeuralTensor(InInputTensorDataSmall.FindChecked(1), TArray<int64>({1,1,1})),
+			FNeuralTensor(ENeuralDataType::Float, {1,1,1}), FNeuralTensor(InInputTensorDataSmall.FindChecked(1))},
+		// X={1,1,9}, W={1,1,5}, B={}, Y={1,1,5}
+		{FNeuralTensor(InInputTensorData.FindChecked(5), TArray<int64>({1,1,5})), FNeuralTensor(InInputTensorDataSmall.FindChecked(5), TArray<int64>({1,1,5})),
+			FNeuralTensor(ENeuralDataType::Float, {1,1,9})},
+		// X={1,3,27}, W={1,3,5}, B={}, Y={1,1,23}
+		{FNeuralTensor(InInputTensorData.FindChecked(23), TArray<int64>({1,1,23})), FNeuralTensor(InInputTensorDataSmall.FindChecked(15), TArray<int64>({1,3,5})),
+			FNeuralTensor(ENeuralDataType::Float, {1,3,27})},
+		// X={1,3,27}, W={2,3,5}, B={}, Y={1,2,23}
+		{FNeuralTensor(InInputTensorData.FindChecked(46), TArray<int64>({1,2,23})), FNeuralTensor(InInputTensorDataSmall.FindChecked(30), TArray<int64>({2,3,5})),
+			FNeuralTensor(ENeuralDataType::Float, {1,3,27})},
+		// X={5,3,20}, W={2,3,5}, B={3}, Y={5,2,16}
+		{FNeuralTensor(InInputTensorData.FindChecked(160), TArray<int64>({5,2,16})), FNeuralTensor(InInputTensorDataSmall.FindChecked(30), TArray<int64>({2,3,5})),
+			FNeuralTensor(ENeuralDataType::Float, {5,3,20}), FNeuralTensor(InInputTensorDataSmall.FindChecked(3))},
+		// ----------------------------------------------------------------- 2D convolutions -----------------------------------------------------------------
+		// X={1,1,1,1}, W={1,1,1,1}, B={}, Y={1,1,1,1} (replicates 1D example)
+		{FNeuralTensor(InInputTensorData.FindChecked(1), TArray<int64>({1,1,1,1})), FNeuralTensor(InInputTensorDataSmall.FindChecked(1), TArray<int64>({1,1,1,1})),
+			FNeuralTensor(ENeuralDataType::Float, {1,1,1,1})},
+		// X={1,1,1,9}, W={1,1,1,5}, B={}, Y={1,1,1,5} (replicates 1D example)
+		{FNeuralTensor(InInputTensorData.FindChecked(5), TArray<int64>({1,1,1,5})), FNeuralTensor(InInputTensorDataSmall.FindChecked(5), TArray<int64>({1,1,1,5})),
+			FNeuralTensor(ENeuralDataType::Float, {1,1,1,9})},
+		// X={5,3,1,20}, W={2,3,1,5}, B={}, Y={5,2,1,16} (replicates 1D example)
+		{FNeuralTensor(InInputTensorData.FindChecked(160), TArray<int64>({5,2,1,16})), FNeuralTensor(InInputTensorDataSmall.FindChecked(30), TArray<int64>({2,3,1,5})),
+			FNeuralTensor(ENeuralDataType::Float, {5,3,1,20})},
+		// X={1,1,1,1}, W={1,1,1,1}, B={1}, Y={1,1,1,1} (replicates 1D example)
+		{FNeuralTensor(InInputTensorData.FindChecked(1), TArray<int64>({1,1,1,1})), FNeuralTensor(InInputTensorDataSmall.FindChecked(1), TArray<int64>({1,1,1,1})),
+			FNeuralTensor(ENeuralDataType::Float, {1,1,1,1}), FNeuralTensor(InInputTensorDataSmall.FindChecked(1))},
+		// X={1,1,5,10}, W={1,1,3,3}, B={}, Y={1,1,3,8}
+		{FNeuralTensor(InInputTensorData.FindChecked(24), TArray<int64>({1,1,3,8})), FNeuralTensor(InInputTensorDataSmall.FindChecked(9), TArray<int64>({1,1,3,3})),
+			FNeuralTensor(ENeuralDataType::Float, {1,1,5,10})},
+		// X={1,1,5,10}, W={1,1,3,3}, B={1}, Y={1,1,3,8}
+		{FNeuralTensor(InInputTensorData.FindChecked(24), TArray<int64>({1,1,3,8})), FNeuralTensor(InInputTensorDataSmall.FindChecked(9), TArray<int64>({1,1,3,3})),
+			FNeuralTensor(ENeuralDataType::Float, {1,1,5,10}), FNeuralTensor(InInputTensorDataSmall.FindChecked(1))},
+		// X={2,1,1,1}, W={5,1,1,1}, B={}, Y={2,5,1,1}
+		{FNeuralTensor(InInputTensorData.FindChecked(10), TArray<int64>({2,5,1,1})), FNeuralTensor(InInputTensorDataSmall.FindChecked(5), TArray<int64>({5,1,1,1})),
+			FNeuralTensor(ENeuralDataType::Float, {2,1,1,1})},
+		// X={2,1,1,1}, W={5,1,1,1}, B={1}, Y={2,5,1,1}
+		{FNeuralTensor(InInputTensorData.FindChecked(10), TArray<int64>({2,5,1,1})), FNeuralTensor(InInputTensorDataSmall.FindChecked(5), TArray<int64>({5,1,1,1})),
+			FNeuralTensor(ENeuralDataType::Float, {2,1,1,1}), FNeuralTensor(InInputTensorDataSmall.FindChecked(1))},
+		// X={2,3,1,1}, W={5,3,1,1}, B={}, Y={2,5,1,1}
+		{FNeuralTensor(InInputTensorData.FindChecked(10), TArray<int64>({2,5,1,1})), FNeuralTensor(InInputTensorDataSmall.FindChecked(15), TArray<int64>({5,3,1,1})),
+			FNeuralTensor(ENeuralDataType::Float, {2,3,1,1})},
+		// X={2,3,1,1}, W={5,3,1,1}, B={3}, Y={2,5,1,1}
+		{FNeuralTensor(InInputTensorData.FindChecked(10), TArray<int64>({2,5,1,1})), FNeuralTensor(InInputTensorDataSmall.FindChecked(15), TArray<int64>({5,3,1,1})),
+			FNeuralTensor(ENeuralDataType::Float, {2,3,1,1}), FNeuralTensor(InInputTensorDataSmall.FindChecked(3))},
+		// ----------------------------------------------------------------- 3D convolutions -----------------------------------------------------------------
+		// X={1,1,1,1,9}, W={1,1,1,1,5}, B={1}, Y={1,1,1,1,5} (replicates 1D example)
+		{FNeuralTensor(InInputTensorData.FindChecked(5), TArray<int64>({1,1,1,1,5})), FNeuralTensor(InInputTensorDataSmall.FindChecked(5), TArray<int64>({1,1,1,1,5})),
+			FNeuralTensor(ENeuralDataType::Float, {1,1,1,1,9}), FNeuralTensor(InInputTensorDataSmall.FindChecked(1))},
+		// X={1,1,1,9,1}, W={1,1,1,5,1}, B={}, Y={1,1,1,5,1} (replicates 1D example)
+		{FNeuralTensor(InInputTensorData.FindChecked(5), TArray<int64>({1,1,1,5,1})), FNeuralTensor(InInputTensorDataSmall.FindChecked(5), TArray<int64>({1,1,1,5,1})),
+			FNeuralTensor(ENeuralDataType::Float, {1,1,1,9,1})},
+		// X={1,1,1,1,1}, W={1,1,1,1,1}, B={}, Y={1,1,1,1,1} (replicates 1D example)
+		{FNeuralTensor(InInputTensorData.FindChecked(1), TArray<int64>({1,1,1,1,1})), FNeuralTensor(InInputTensorDataSmall.FindChecked(1), TArray<int64>({1,1,1,1,1})),
+			FNeuralTensor(ENeuralDataType::Float, {1,1,1,1,1})},
+		// X={1,1,6,5,5}, W={1,1,3,3,3}, B={}, Y={1,1,4,3,3}
+		{FNeuralTensor(InInputTensorData.FindChecked(36), TArray<int64>({1,1,4,3,3})), FNeuralTensor(InInputTensorDataSmall.FindChecked(27), TArray<int64>({1,1,3,3,3})),
+			FNeuralTensor(ENeuralDataType::Float, {1,1,6,5,5})},
+		// X={2,2,3,5,5}, W={9,2,3,3,3}, B={2}, Y={2,9,1,3,3}
+		{FNeuralTensor(InInputTensorData.FindChecked(162), TArray<int64>({2,9,1,3,3})), FNeuralTensor(InInputTensorDataSmall.FindChecked(486), TArray<int64>({9,2,3,3,3})),
+			FNeuralTensor(ENeuralDataType::Float, {2,2,3,5,5}), FNeuralTensor(InInputTensorDataSmall.FindChecked(2))},
+		// ----------------------------------------------------------------- 4D convolutions -----------------------------------------------------------------
+		// X={1,1,1,1,1,1}, W={1,1,1,1,1,1}, B={1}, Y={1,1,1,1,1,1} (replicates 1D example)
+		{FNeuralTensor(InInputTensorData.FindChecked(1), TArray<int64>({1,1,1,1,1,1})), FNeuralTensor(InInputTensorDataSmall.FindChecked(1), TArray<int64>({1,1,1,1,1,1})),
+			FNeuralTensor(ENeuralDataType::Float, {1,1,1,1,1,1}), FNeuralTensor(InInputTensorDataSmall.FindChecked(1))},
+		// X={1,1,1,1,1,9}, W={1,1,1,1,1,5}, B={}, Y={1,1,1,1,1,5} (replicates 1D example)
+		{FNeuralTensor(InInputTensorData.FindChecked(5), TArray<int64>({1,1,1,1,1,5})), FNeuralTensor(InInputTensorDataSmall.FindChecked(5), TArray<int64>({1,1,1,1,1,5})),
+			FNeuralTensor(ENeuralDataType::Float, {1,1,1,1,1,9})},
+		// X={1,1,3,4,5,5}, W={1,1,3,3,3,3}, B={}, Y={1,1,1,2,3,3}
+		{FNeuralTensor(InInputTensorData.FindChecked(18), TArray<int64>({1,1,1,2,3,3})), FNeuralTensor(InInputTensorDataSmall.FindChecked(81), TArray<int64>({1,1,3,3,3,3})),
+			FNeuralTensor(ENeuralDataType::Float, {1,1,3,4,5,5})},
+		// X={3,3,4,5,3,5}, W={2,3,3,3,3,3}, B={3}, Y={3,2,2,3,1,3}
+		{FNeuralTensor(InInputTensorData.FindChecked(108), TArray<int64>({3,2,2,3,1,3})), FNeuralTensor(InInputTensorDataSmall.FindChecked(486), TArray<int64>({2,3,3,3,3,3})),
+			FNeuralTensor(ENeuralDataType::Float, {3,3,4,5,3,5}), FNeuralTensor(InInputTensorDataSmall.FindChecked(3))},
+		// ----------------------------------------------------------------- 8D convolutions -----------------------------------------------------------------
+		// X={1,1, 1,1,1,1,1,1,1,1}, W={1,1, 1,1,1,1,1,1,1,1}, B={1}, Y={1,1, 1,1,1,1,1,1,1,1} (replicates 1D example)
+		{FNeuralTensor(InInputTensorData.FindChecked(1), TArray<int64>({1,1, 1,1,1,1,1,1,1,1})), FNeuralTensor(InInputTensorDataSmall.FindChecked(1), TArray<int64>({1,1, 1,1,1,1,1,1,1,1})),
+			FNeuralTensor(ENeuralDataType::Float, {1,1, 1,1,1,1,1,1,1,1}), FNeuralTensor(InInputTensorDataSmall.FindChecked(1))},
+		// X={1,1, 1,1,1,1,1,1,1,9}, W={1,1, 1,1,1,1,1,1,1,5}, B={}, Y={1,1, 1,1,1,1,1,1,1,5} (replicates 1D example)
+		{FNeuralTensor(InInputTensorData.FindChecked(5), TArray<int64>({1,1, 1,1,1,1,1,1,1,5})), FNeuralTensor(InInputTensorDataSmall.FindChecked(5), TArray<int64>({1,1, 1,1,1,1,1,1,1,5})),
+			FNeuralTensor(ENeuralDataType::Float, {1,1, 1,1,1,1,1,1,1,9})},
+		// X={1,1, 1,1,1,1,3,4,5,5}, W={1,1, 1,1,1,1,3,3,3,3}, B={}, Y={1,1, 1,1,1,1,1,2,3,3} (replicates 4D example)
+		{FNeuralTensor(InInputTensorData.FindChecked(18), TArray<int64>({1,1, 1,1,1,1,1,2,3,3})), FNeuralTensor(InInputTensorDataSmall.FindChecked(81), TArray<int64>({1,1, 1,1,1,1,3,3,3,3})),
+			FNeuralTensor(ENeuralDataType::Float, {1,1, 1,1,1,1,3,4,5,5})},
+		// X={3,3, 1,1,1,1,4,5,3,5}, W={2,3, 1,1,1,1,3,3,3,3}, B={3}, Y={3,2, 1,1,1,1,2,3,1,3} (replicates 4D example)
+		{FNeuralTensor(InInputTensorData.FindChecked(108), TArray<int64>({3,2, 1,1,1,1,2,3,1,3})), FNeuralTensor(InInputTensorDataSmall.FindChecked(486), TArray<int64>({2,3, 1,1,1,1,3,3,3,3})),
+			FNeuralTensor(ENeuralDataType::Float, {3,3, 1,1,1,1,4,5,3,5}), FNeuralTensor(InInputTensorDataSmall.FindChecked(3))},
+		// ----------------------------------------------------------------- Grouped convolutions -----------------------------------------------------------------
+		// X={1,4,1}, W={8,4/2,1}, B={}, Y={1,8,1}, Group = 2
+		{FNeuralTensor(TArray<float>({1,1,1,1, 1,1,1,1}), TArray<int64>({1,8,1})),
+			FNeuralTensor(TArray<float>({1,1,1,1,1,1,1,1, 1,1,1,1,1,1,1,1, 1,1,1,1,1,1,1,1, 1,1,1,1,1,1,1,1}), TArray<int64>({8,4,1})),
+			FNeuralTensor(ENeuralDataType::Float, {1,4,1})},
+		{FNeuralTensor(TArray<float>({1,1,1,1, 1,1,1,1}), TArray<int64>({1,8,1})),
+			FNeuralTensor(TArray<float>({1,1,1,1,1,1,1,1, 1,1,1,1,1,1,1,1}), TArray<int64>({8,2,1})),
+			FNeuralTensor(ENeuralDataType::Float, {1,4,1}), FNeuralTensor(), /*Group*/FNeuralTensor(TArray<float>({2}))},
+		{FNeuralTensor(InInputTensorData.FindChecked(8), TArray<int64>({1,8,1})), FNeuralTensor(InInputTensorDataSmall.FindChecked(16), TArray<int64>({8,2,1})),
+			FNeuralTensor(ENeuralDataType::Float, {1,4,1}), FNeuralTensor(), /*Group*/FNeuralTensor(TArray<float>({2}))},
+		// X={1,9,9,6}, W={6,9/3,9,3}, B={}, Y={1,6,1,4}, Group = 3
+		{FNeuralTensor(InInputTensorData.FindChecked(24), TArray<int64>({1,6,1,4})), FNeuralTensor(InInputTensorDataSmall.FindChecked(486), TArray<int64>({6,3,9,3})),
+			FNeuralTensor(ENeuralDataType::Float, {1,9,9,6}), FNeuralTensor(), /*Group*/FNeuralTensor(TArray<float>({3}))},
+		// X={1,8,1,1,1}, W={4,8/4,1,1,1}, B={8}, Y={1,4,1,1,1}, Group = 4
+		{FNeuralTensor(InInputTensorData.FindChecked(4), TArray<int64>({1,4,1,1,1})), FNeuralTensor(InInputTensorDataSmall.FindChecked(8), TArray<int64>({4,2,1,1,1})),
+			FNeuralTensor(ENeuralDataType::Float, {1,8,1,1,1}), FNeuralTensor(InInputTensorDataSmall.FindChecked(8)), /*Group*/FNeuralTensor(TArray<float>({4}))},
+		// ----------------------------------------------------------------- AutoPad == SameLower/Upper -----------------------------------------------------------------
+		// X={1,1,9}, W={1,1,4}, B={}, Y={1,1,9}, AutoPad = EAutoPad::SameLower || EAutoPad::SameUpper
+		{FNeuralTensor(InInputTensorData.FindChecked(9), TArray<int64>({1,1,9})), FNeuralTensor(InInputTensorDataSmall.FindChecked(4), TArray<int64>({1,1,4})),
+			FNeuralTensor(ENeuralDataType::Float, {1,1,9}), FNeuralTensor(), /*Group*/FNeuralTensor(), /*AutoPad*/FNeuralTensor(TArray<float>({float(FConvBaseOperator::EAutoPad::SameLower)}))},
+		//                                        AutoPad = EAutoPad::SameLower || EAutoPad::SameUpper
+		{FNeuralTensor(InInputTensorData.FindChecked(9), TArray<int64>({1,1,9})), FNeuralTensor(InInputTensorDataSmall.FindChecked(4), TArray<int64>({1,1,4})),
+			FNeuralTensor(ENeuralDataType::Float, {1,1,9}), FNeuralTensor(), /*Group*/FNeuralTensor(), /*AutoPad*/FNeuralTensor(TArray<float>({float(FConvBaseOperator::EAutoPad::SameUpper)})) },
+		// X={1,1,9}, W={1,1,9}, B={}, Y={1,1,9}, AutoPad = EAutoPad::SameLower || EAutoPad::SameUpper
+		{FNeuralTensor(InInputTensorData.FindChecked(9), TArray<int64>({1,1,9})), FNeuralTensor(InInputTensorDataSmall.FindChecked(9), TArray<int64>({1,1,9})),
+			FNeuralTensor(ENeuralDataType::Float, {1,1,9}), FNeuralTensor(), /*Group*/FNeuralTensor(), /*AutoPad*/FNeuralTensor(TArray<float>({float(FConvBaseOperator::EAutoPad::SameLower)}))},
+		//                                        AutoPad = EAutoPad::SameLower || EAutoPad::SameUpper
+		{FNeuralTensor(InInputTensorData.FindChecked(9), TArray<int64>({1,1,9})), FNeuralTensor(InInputTensorDataSmall.FindChecked(9), TArray<int64>({1,1,9})),
+			FNeuralTensor(ENeuralDataType::Float, {1,1,9}), FNeuralTensor(), /*Group*/FNeuralTensor(), /*AutoPad*/FNeuralTensor(TArray<float>({float(FConvBaseOperator::EAutoPad::SameUpper)}))},
+		// X={1,1,9,9}, W={1,1,3,3}, B={}, Y={1,1,9,9}, AutoPad = EAutoPad::SameLower || EAutoPad::SameUpper
+		{FNeuralTensor(InInputTensorData.FindChecked(81), TArray<int64>({1,1,9,9})), FNeuralTensor(InInputTensorDataSmall.FindChecked(9), TArray<int64>({1,1,3,3})),
+			FNeuralTensor(ENeuralDataType::Float, {1,1,9,9}), FNeuralTensor(), /*Group*/FNeuralTensor(), /*AutoPad*/FNeuralTensor(TArray<float>({float(FConvBaseOperator::EAutoPad::SameLower)}))},
+		//                                              AutoPad = EAutoPad::SameLower || EAutoPad::SameUpper
+		{FNeuralTensor(InInputTensorData.FindChecked(81), TArray<int64>({1,1,9,9})), FNeuralTensor(InInputTensorDataSmall.FindChecked(9), TArray<int64>({1,1,3,3})),
+			FNeuralTensor(ENeuralDataType::Float, {1,1,9,9}), FNeuralTensor(), /*Group*/FNeuralTensor(), /*AutoPad*/FNeuralTensor(TArray<float>({float(FConvBaseOperator::EAutoPad::SameUpper)}))},
+		// X={1,1,9,9}, W={1,1,5,5}, B={}, Y={1,1,9,9}, AutoPad = EAutoPad::SameLower || EAutoPad::SameUpper
+		{FNeuralTensor(InInputTensorData.FindChecked(81), TArray<int64>({1,1,9,9})), FNeuralTensor(InInputTensorDataSmall.FindChecked(25), TArray<int64>({1,1,5,5})),
+			FNeuralTensor(ENeuralDataType::Float, {1,1,9,9}), FNeuralTensor(), /*Group*/FNeuralTensor(), /*AutoPad*/FNeuralTensor(TArray<float>({float(FConvBaseOperator::EAutoPad::SameLower)}))},
+		//                                              AutoPad = EAutoPad::SameLower || EAutoPad::SameUpper
+		{FNeuralTensor(InInputTensorData.FindChecked(81), TArray<int64>({1,1,9,9})), FNeuralTensor(InInputTensorDataSmall.FindChecked(25), TArray<int64>({1,1,5,5})),
+			FNeuralTensor(ENeuralDataType::Float, {1,1,9,9}), FNeuralTensor(), /*Group*/FNeuralTensor(), /*AutoPad*/FNeuralTensor(TArray<float>({float(FConvBaseOperator::EAutoPad::SameUpper)}))},
+		// X={1,1,5,5,4}, W={1,1,3,3,3}, B={}, Y={1,1,5,5,4}, AutoPad = EAutoPad::SameLower || EAutoPad::SameUpper
+		{FNeuralTensor(InInputTensorData.FindChecked(100), TArray<int64>({1,1,5,5,4})), FNeuralTensor(InInputTensorDataSmall.FindChecked(27), TArray<int64>({1,1,3,3,3})),
+			FNeuralTensor(ENeuralDataType::Float, {1,1,5,5,4}), FNeuralTensor(), /*Group*/FNeuralTensor(), /*AutoPad*/FNeuralTensor(TArray<float>({float(FConvBaseOperator::EAutoPad::SameLower)}))},
+		//                                                    AutoPad = EAutoPad::SameLower || EAutoPad::SameUpper
+		{FNeuralTensor(InInputTensorData.FindChecked(100), TArray<int64>({1,1,5,5,4})), FNeuralTensor(InInputTensorDataSmall.FindChecked(27), TArray<int64>({1,1,3,3,3})),
+			FNeuralTensor(ENeuralDataType::Float, {1,1,5,5,4}), FNeuralTensor(), /*Group*/FNeuralTensor(), /*AutoPad*/FNeuralTensor(TArray<float>({float(FConvBaseOperator::EAutoPad::SameUpper)}))},
+		// X={1,1,1,1,5,5,4}, W={1,1,1,1,3,3,3}, B={}, Y={1,1,1,1,5,5,4}, AutoPad = EAutoPad::SameLower || EAutoPad::SameUpper
+		{FNeuralTensor(InInputTensorData.FindChecked(100), TArray<int64>({1,1,1,1,5,5,4})), FNeuralTensor(InInputTensorDataSmall.FindChecked(27), TArray<int64>({1,1,1,1,3,3,3})),
+			FNeuralTensor(ENeuralDataType::Float, {1,1,1,1,5,5,4}), FNeuralTensor(), /*Group*/FNeuralTensor(), /*AutoPad*/FNeuralTensor(TArray<float>({float(FConvBaseOperator::EAutoPad::SameLower)}))},
+		//                                                    AutoPad = EAutoPad::SameLower || EAutoPad::SameUpper
+		{FNeuralTensor(InInputTensorData.FindChecked(100), TArray<int64>({1,1,5,5,4,1,1})), FNeuralTensor(InInputTensorDataSmall.FindChecked(27), TArray<int64>({1,1,3,3,3,1,1})),
+			FNeuralTensor(ENeuralDataType::Float, {1,1,5,5,4,1,1}), FNeuralTensor(), /*Group*/FNeuralTensor(), /*AutoPad*/FNeuralTensor(TArray<float>({float(FConvBaseOperator::EAutoPad::SameUpper)}))},
+		// X={2,3,5,2,3,3,5}, W={2,3,5,2,1,1,5}, B={}, Y={2,2,5,2,3,3,5}, AutoPad = EAutoPad::SameLower || EAutoPad::SameUpper
+		{FNeuralTensor(InInputTensorData.FindChecked(1800), TArray<int64>({2,2,5,2,3,3,5})), FNeuralTensor(InInputTensorDataSmall.FindChecked(300), TArray<int64>({2,3,5,2,1,1,5})),
+			FNeuralTensor(ENeuralDataType::Float, {2,3,5,2,3,3,5}), FNeuralTensor(), /*Group*/FNeuralTensor(), /*AutoPad*/FNeuralTensor(TArray<float>({float(FConvBaseOperator::EAutoPad::SameLower)}))},
+		//																  AutoPad = EAutoPad::SameLower || EAutoPad::SameUpper
+		{FNeuralTensor(InInputTensorData.FindChecked(1800), TArray<int64>({2,2,5,2,3,3,5})), FNeuralTensor(InInputTensorDataSmall.FindChecked(300), TArray<int64>({2,3,5,2,1,1,5})),
+			FNeuralTensor(ENeuralDataType::Float, {2,3,5,2,3,3,5}), FNeuralTensor(), /*Group*/FNeuralTensor(), /*AutoPad*/FNeuralTensor(TArray<float>({float(FConvBaseOperator::EAutoPad::SameUpper)}))},
+		// ----------------------------------------------------------------- AutoPad = EAutoPad::NotSet (i.e., manual padding) -----------------------------------------------------------------
+		// X={1,1,9} + {0+3}, W={1,1,4}, B={}, Y={1,1,9}
+		{FNeuralTensor(InInputTensorData.FindChecked(9), TArray<int64>({1,1,9})), FNeuralTensor(InInputTensorDataSmall.FindChecked(4), TArray<int64>({1,1,4})),
+			FNeuralTensor(ENeuralDataType::Float, {1,1,9}), FNeuralTensor(), /*Group*/FNeuralTensor(), /*AutoPad*/FNeuralTensor(), /*Pads*/FNeuralTensor(TArray<float>({0,3}))},
+		// X={1,1,9,9} + {1+0,2+1}, W={1,1,3,3}, B={}, Y={1,1,8,10}
+		{FNeuralTensor(InInputTensorData.FindChecked(80), TArray<int64>({1,1,8,10})), FNeuralTensor(InInputTensorDataSmall.FindChecked(9), TArray<int64>({1,1,3,3})),
+			FNeuralTensor(ENeuralDataType::Float, {1,1,9,9}), FNeuralTensor(), /*Group*/FNeuralTensor(), /*AutoPad*/FNeuralTensor(), /*Pads*/FNeuralTensor(TArray<float>({1,0, 2,1}))},
+		// X={1,1,5,5,4} + {0+1,2+2,2+0}, W={1,1,3,3,3}, B={}, Y={1,1,4,7,4}
+		{FNeuralTensor(InInputTensorData.FindChecked(112), TArray<int64>({1,1,4,7,4})), FNeuralTensor(InInputTensorDataSmall.FindChecked(27), TArray<int64>({1,1,3,3,3})),
+			FNeuralTensor(ENeuralDataType::Float, {1,1,5,5,4}), FNeuralTensor(), /*Group*/FNeuralTensor(), /*AutoPad*/FNeuralTensor(), /*Pads*/FNeuralTensor(TArray<float>({0,1, 2,2, 2,0}))},
+		// X={1,1,1,1,5,5,4} + {2+2,2+2,0,1+2,2+1}, W={1,1,3,3,3,3,3}, B={}, Y={1,1,3,3,3,6,5}
+		{FNeuralTensor(InInputTensorData.FindChecked(810), TArray<int64>({1,1,3,3,3,6,5})), FNeuralTensor(InInputTensorDataSmall.FindChecked(243), TArray<int64>({1,1,3,3,3,3,3})),
+			FNeuralTensor(ENeuralDataType::Float, {1,1,1,1,5,5,4}), FNeuralTensor(), /*Group*/FNeuralTensor(), /*AutoPad*/FNeuralTensor(), /*Pads*/FNeuralTensor(TArray<float>({2,2, 2,2, 0,0, 1,2, 2,1}))},
+		// ----------------------------------------------------------------- Dilations != 1 -----------------------------------------------------------------
+		// X={1,1,9}, W={1,1,4}, B={}, Y={1,1,3}
+		{FNeuralTensor(InInputTensorData.FindChecked(3), TArray<int64>({1,1,3})), FNeuralTensor(InInputTensorDataSmall.FindChecked(4), TArray<int64>({1,1,4})),
+			FNeuralTensor(ENeuralDataType::Float, {1,1,9}), FNeuralTensor(), /*Group*/FNeuralTensor(), /*AutoPad*/FNeuralTensor(), /*Pads*/FNeuralTensor(), /*Dilations*/FNeuralTensor(TArray<float>({2}))},
+		// X={1,1,9,9}, W={1,1,3,3}, B={}, Y={1,1,3,5}
+		{FNeuralTensor(InInputTensorData.FindChecked(15), TArray<int64>({1,1,3,5})), FNeuralTensor(InInputTensorDataSmall.FindChecked(9), TArray<int64>({1,1,3,3})),
+			FNeuralTensor(ENeuralDataType::Float, {1,1,9,9}), FNeuralTensor(), /*Group*/FNeuralTensor(), /*AutoPad*/FNeuralTensor(), /*Pads*/FNeuralTensor(), /*Dilations*/FNeuralTensor(TArray<float>({3,2}))},
+		// X={1,1,5,4,5}, W={1,1,3,3,3}, B={}, Y={1,1,1,2,1}
+		{FNeuralTensor(InInputTensorData.FindChecked(2), TArray<int64>({1,1,1,2,1})), FNeuralTensor(InInputTensorDataSmall.FindChecked(27), TArray<int64>({1,1,3,3,3})),
+			FNeuralTensor(ENeuralDataType::Float, {1,1,5,4,5}), FNeuralTensor(), /*Group*/FNeuralTensor(), /*AutoPad*/FNeuralTensor(), /*Pads*/FNeuralTensor(), /*Dilations*/FNeuralTensor(TArray<float>({2,1,2}))},
+		// X={1,1,1,1,5,4,5}, W={1,1,1,1,3,3,3}, B={}, Y={1,1,1,1,1,2,1}
+		{FNeuralTensor(InInputTensorData.FindChecked(2), TArray<int64>({1,1,1,1,1,2,1})), FNeuralTensor(InInputTensorDataSmall.FindChecked(27), TArray<int64>({1,1,1,1,3,3,3})),
+			FNeuralTensor(ENeuralDataType::Float, {1,1,1,1,5,4,5}), FNeuralTensor(), /*Group*/FNeuralTensor(), /*AutoPad*/FNeuralTensor(), /*Pads*/FNeuralTensor(), /*Dilations*/FNeuralTensor(TArray<float>({99,10,2,1,2}))},
+		// ----------------------------------------------------------------- Strides != 1 -----------------------------------------------------------------
+		// Required manual padding = 0 to avoid this unit tester to accidentally auto-trigger AutoPad == SameLower/Upper
+		// X={1,1,9}, W={1,1,4}, B={}, Y={1,1,2}
+		{FNeuralTensor(InInputTensorData.FindChecked(2), TArray<int64>({1,1,2})), FNeuralTensor(InInputTensorDataSmall.FindChecked(4), TArray<int64>({1,1,4})),
+			FNeuralTensor(ENeuralDataType::Float, {1,1,9}), FNeuralTensor(), /*Group*/FNeuralTensor(), /*AutoPad*/FNeuralTensor(), /*Pads*/FNeuralTensor(),
+			/*Dilations*/FNeuralTensor(), /*Strides*/FNeuralTensor(TArray<float>({5}))},
+		// X={1,1,9,9}, W={1,1,3,3}, B={}, Y={1,1,1,3}
+		{FNeuralTensor(InInputTensorData.FindChecked(3), TArray<int64>({1,1,1,3})), FNeuralTensor(InInputTensorDataSmall.FindChecked(9), TArray<int64>({1,1,3,3})),
+			FNeuralTensor(ENeuralDataType::Float, {1,1,9,9}), FNeuralTensor(), /*Group*/FNeuralTensor(), /*AutoPad*/FNeuralTensor(), /*Pads*/FNeuralTensor(),
+			/*Dilations*/FNeuralTensor(), /*Strides*/FNeuralTensor(TArray<float>({7,3}))},
+		// X={1,1,5,4,5}, W={1,1,3,3,3}, B={}, Y={1,1,2,2,1}
+		{FNeuralTensor(InInputTensorData.FindChecked(4), TArray<int64>({1,1,2,2,1})), FNeuralTensor(InInputTensorDataSmall.FindChecked(27), TArray<int64>({1,1,3,3,3})),
+			FNeuralTensor(ENeuralDataType::Float, {1,1,5,4,5}), FNeuralTensor(), /*Group*/FNeuralTensor(), /*AutoPad*/FNeuralTensor(), /*Pads*/FNeuralTensor(),
+			/*Dilations*/FNeuralTensor(), /*Strides*/FNeuralTensor(TArray<float>({2,1,3}))},
+		// X={1,1,1,1,5,4,5}, W={1,1,1,1,3,3,3}, B={}, Y={1,1,1,1,2,2,1}
+		{FNeuralTensor(InInputTensorData.FindChecked(4), TArray<int64>({1,1,1,1,2,2,1})), FNeuralTensor(InInputTensorDataSmall.FindChecked(27), TArray<int64>({1,1,1,1,3,3,3})),
+			FNeuralTensor(ENeuralDataType::Float, {1,1,1,1,5,4,5}), FNeuralTensor(), /*Group*/FNeuralTensor(), /*AutoPad*/FNeuralTensor(), /*Pads*/FNeuralTensor(),
+			/*Dilations*/FNeuralTensor(), /*Strides*/FNeuralTensor(TArray<float>({1,1,2,1,3}))},
+		// ----------------------------------------------------------------- Strides != 1, Multi-channel -----------------------------------------------------------------
+		// X={5,5,25}, W={2,5,9}, B={}, Y={5,2,9}
+		{FNeuralTensor(InInputTensorData.FindChecked(90), TArray<int64>({5,2,9})), FNeuralTensor(InInputTensorDataSmall.FindChecked(90), TArray<int64>({2,5,9})),
+			FNeuralTensor(ENeuralDataType::Float, {5,5,25}), FNeuralTensor(), /*Group*/FNeuralTensor(), /*AutoPad*/FNeuralTensor(), /*Pads*/FNeuralTensor(),
+			/*Dilations*/FNeuralTensor(), /*Strides*/FNeuralTensor(TArray<float>({2}))},
+		// X={5,5,5,5}, W={2,5,3,3}, B={}, Y={5,2,2,2}
+		{FNeuralTensor(InInputTensorData.FindChecked(40), TArray<int64>({5,2,2,2})), FNeuralTensor(InInputTensorDataSmall.FindChecked(90), TArray<int64>({2,5,3,3})),
+			FNeuralTensor(ENeuralDataType::Float, {5,5,5,5}), FNeuralTensor(), /*Group*/FNeuralTensor(), /*AutoPad*/FNeuralTensor(), /*Pads*/FNeuralTensor(),
+			/*Dilations*/FNeuralTensor(), /*Strides*/FNeuralTensor(TArray<float>({2,2}))},
+		// X={5,5,5,5,1}, W={2,5,3,3,1}, B={}, Y={5,2,2,2,1}
+		{FNeuralTensor(InInputTensorData.FindChecked(40), TArray<int64>({5,2,2,2,1})), FNeuralTensor(InInputTensorDataSmall.FindChecked(90), TArray<int64>({2,5,3,3,1})),
+			FNeuralTensor(ENeuralDataType::Float, {5,5,5,5,1}), FNeuralTensor(), /*Group*/FNeuralTensor(), /*AutoPad*/FNeuralTensor(), /*Pads*/FNeuralTensor(),
+			/*Dilations*/FNeuralTensor(), /*Strides*/FNeuralTensor(TArray<float>({2,2,1}))},
+		// X={5,5,5,5,1,1,1}, W={2,5,3,3,1,1,1}, B={}, Y={5,2,2,2,1,1,1}
+		{FNeuralTensor(InInputTensorData.FindChecked(40), TArray<int64>({5,2,2,2,1,1,1})), FNeuralTensor(InInputTensorDataSmall.FindChecked(90), TArray<int64>({2,5,3,3,1,1,1})),
+			FNeuralTensor(ENeuralDataType::Float, {5,5,5,5,1,1,1}), FNeuralTensor(), /*Group*/FNeuralTensor(), /*AutoPad*/FNeuralTensor(), /*Pads*/FNeuralTensor(),
+			/*Dilations*/FNeuralTensor(), /*Strides*/FNeuralTensor(TArray<float>({2,2,1,1,1}))},
+		// ----------------------------------------------------------------- Manual padding or AutoPad == SameLower/Upper, Dilations != 1, Strides != 1 -----------------------------------------------------------------
+		// X={1,1,18}, W={1,1,4}, B={}, Y={1,1,9}, AutoPad = EAutoPad::SameLower || EAutoPad::SameUpper
+		{FNeuralTensor(InInputTensorData.FindChecked(9), TArray<int64>({1,1,9})), FNeuralTensor(InInputTensorDataSmall.FindChecked(4), TArray<int64>({1,1,4})),
+			FNeuralTensor(ENeuralDataType::Float, {1,1,18}), FNeuralTensor(), /*Group*/FNeuralTensor(), /*AutoPad*/FNeuralTensor(), /*Pads*/FNeuralTensor(),
+			/*Dilations*/FNeuralTensor(TArray<float>({2})), /*Strides*/FNeuralTensor(TArray<float>({2})), /*OutputPadding*/FNeuralTensor(), /*OutputShape*/FNeuralTensor(TArray<float>({ 18 }))},
+		//										   AutoPad = EAutoPad::SameUpper || EAutoPad::SameLower
+		{FNeuralTensor(InInputTensorData.FindChecked(9), TArray<int64>({1,1,9})), FNeuralTensor(InInputTensorDataSmall.FindChecked(4), TArray<int64>({1,1,4})),
+			FNeuralTensor(ENeuralDataType::Float, {1,1,18}), FNeuralTensor(), /*Group*/FNeuralTensor(), /*AutoPad*/FNeuralTensor(), /*Pads*/FNeuralTensor(),
+			/*Dilations*/FNeuralTensor(TArray<float>({2})), /*Strides*/FNeuralTensor(TArray<float>({2})), /*OutputPadding*/FNeuralTensor(), /*OutputShape*/FNeuralTensor(TArray<float>({ 18 }))},
+		// X={1,1,100}, W={1,1,3}, B={}, Y={1,1,21}
+		{FNeuralTensor(InInputTensorData.FindChecked(21), TArray<int64>({1,1,21})), FNeuralTensor(InInputTensorDataSmall.FindChecked(3), TArray<int64>({1,1,3})),
+			FNeuralTensor(ENeuralDataType::Float, {1,1,100}), FNeuralTensor(), /*Group*/FNeuralTensor(), /*AutoPad*/FNeuralTensor(), /*Pads*/FNeuralTensor(TArray<float>({10,5})),
+			/*Dilations*/FNeuralTensor(TArray<float>({5})), /*Strides*/FNeuralTensor(TArray<float>({5})), /*OutputPadding*/FNeuralTensor(TArray<float>({ 4 }))},
+		// X={1,1,9,9}+{2+1,0+2}, W={1,1,3,3}, B={}, Y={1,1,8,9}
+		{FNeuralTensor(InInputTensorData.FindChecked(72), TArray<int64>({1,1,8,9})), FNeuralTensor(InInputTensorDataSmall.FindChecked(9), TArray<int64>({1,1,3,3})),
+			FNeuralTensor(ENeuralDataType::Float, {1,1,9,9}), FNeuralTensor(), /*Group*/FNeuralTensor(), /*AutoPad*/FNeuralTensor(), /*Pads*/FNeuralTensor(TArray<float>({2,1,0,2})),
+			/*Dilations*/FNeuralTensor(TArray<float>({2,1})), /*Strides*/FNeuralTensor(TArray<float>({1,1}))},
+		// X={1,1,9,9}+{2+1,0+2}, W={1,1,3,3}, B={}, Y={1,1,4,5}
+		{FNeuralTensor(InInputTensorData.FindChecked(20), TArray<int64>({1,1,4,5})), FNeuralTensor(InInputTensorDataSmall.FindChecked(9), TArray<int64>({1,1,3,3})),
+			FNeuralTensor(ENeuralDataType::Float, {1,1,9,9}), FNeuralTensor(), /*Group*/FNeuralTensor(), /*AutoPad*/FNeuralTensor(), /*Pads*/FNeuralTensor(TArray<float>({2,1,0,2})),
+			/*Dilations*/FNeuralTensor(TArray<float>({2,1})), /*Strides*/FNeuralTensor(TArray<float>({2,2})), /*OutputPadding*/FNeuralTensor(TArray<float>({ 1,0 }))},
+		// X={1,1,5,4,5}, W={1,1,3,3,3}, B={}, Y={1,1,3,2,2}
+		{FNeuralTensor(InInputTensorData.FindChecked(12), TArray<int64>({1,1,3,2,2})), FNeuralTensor(InInputTensorDataSmall.FindChecked(27), TArray<int64>({1,1,3,3,3})),
+			FNeuralTensor(ENeuralDataType::Float, {1,1,5,4,5}), FNeuralTensor(), /*Group*/FNeuralTensor(), /*AutoPad*/FNeuralTensor(), /*Pads*/FNeuralTensor(TArray<float>({3,3,2,3,3,2})),
+			/*Dilations*/FNeuralTensor(TArray<float>({3,3,3})), /*Strides*/FNeuralTensor(TArray<float>({2,2,2})), /*OutputPadding*/FNeuralTensor(TArray<float>({ 0,0,1 }))},
+		// X={1,1,1,1,5,4,5}, W={1,1,1,1,3,3,3}, B={}, Y={1,1,1,1,3,2,2}
+		{FNeuralTensor(InInputTensorData.FindChecked(12), TArray<int64>({1,1,1,1,3,2,2})), FNeuralTensor(InInputTensorDataSmall.FindChecked(27), TArray<int64>({1,1,1,1,3,3,3})),
+			FNeuralTensor(ENeuralDataType::Float, {1,1,1,1,5,4,5}), FNeuralTensor(), /*Group*/FNeuralTensor(), /*AutoPad*/FNeuralTensor(), /*Pads*/FNeuralTensor(TArray<float>({0,0,0,0,3,3,2,3,3,2})),
+			/*Dilations*/FNeuralTensor(TArray<float>({1,1,3,3,3})), /*Strides*/FNeuralTensor(TArray<float>({1,1,2,2,2})), /*OutputPadding*/FNeuralTensor(TArray<float>({0,0,0,0,1}))}
+		});
+	TArray<TArray<FNeuralTensor*>> InputTensors;
+	TArray<TArray<FNeuralTensor*>> OutputTensors;
+	TArray<TSharedPtr<FNeuralOperator>> Operators;
+	for (int32 TensorIndex = 0; TensorIndex < Tensors.Num(); ++TensorIndex)
+	{
+		TArray<FNeuralTensor>& TensorsIndex = Tensors[TensorIndex];
+		if (TensorsIndex.Num() > 3 && TensorsIndex[3].Num() > 0)
+		{
+			InputTensors.Push({ &TensorsIndex[0], &TensorsIndex[1], &TensorsIndex[3] });
+		}
+		else
+		{
+			InputTensors.Push({ &TensorsIndex[0], &TensorsIndex[1] });
+		}
+		OutputTensors.Push({ &TensorsIndex[2] });
+		// Set Group - Hacky but simple way of getting group value
+		//const int64 Group = (TensorsIndex.Num() > 4 && TensorsIndex[4].Num() > 0 ? FMath::RoundToInt(TensorsIndex[4].At<float>(0)) : 1);
+		const TArray<int64> GroupArray = FloatTensorToIntArray(TensorsIndex, 4);
+		if (GroupArray.Num() > 1)
+		{
+			ensureMsgf(false, TEXT("GroupArray.Num() = %d == 0 | 1 failed."), GroupArray.Num());
+			return;
+		}
+		const int64 Group = (GroupArray.Num() > 0 ? GroupArray[0] : 1);
+		// Set FConvBaseOperator::EAutoPad - Hacky but simple way of getting this TArray
+		FConvBaseOperator::EAutoPad AutoPad = FConvBaseOperator::EAutoPad::NotSet;
+		const TArray<int64> AutoPadArray = FloatTensorToIntArray(TensorsIndex, 5);
+		if (AutoPadArray.Num() > 0)
+		{
+			if (AutoPadArray.Num() != 1)
+			{
+				ensureMsgf(false, TEXT("AutoPadArray should always be 1."));
+				return;
+			}
+			AutoPad = FConvBaseOperator::EAutoPad(AutoPadArray[0]);
+		}
+		// Set Pads - Hacky but simple way of getting this TArray
+		const TArray<int64> Pads = FloatTensorToIntArray(TensorsIndex, 6);
+		// Set Dilations - Hacky but simple way of getting this TArray
+		const TArray<int64> Dilations = FloatTensorToIntArray(TensorsIndex, 7);
+		// Set Strides - Hacky but simple way of getting this TArray
+		const TArray<int64> Strides = FloatTensorToIntArray(TensorsIndex, 8);
+		// Set OutputPadding - Hacky but simple way of getting this TArray
+		const TArray<int64> OutputPadding = FloatTensorToIntArray(TensorsIndex, 9);
+		// Set OutputShape - Hacky but simple way of getting this TArray
+		const TArray<int64> OutputShape = FloatTensorToIntArray(TensorsIndex, 10);
+		// Create and push FConvTransposeOperator
+		{
+			TSharedPtr<FConvTransposeOperator> ConvTransposeOperator;
+			ConvTransposeOperator = MakeShared<FConvTransposeOperator>(AutoPad, Dilations, Group, /*KernelShape*/TArray<int64>(), OutputPadding, OutputShape, Pads, Strides);
+			// Flip weights
+			{
+				UE_LOG(LogNeuralNetworkInferenceQA, Display, TEXT("Flipping FConvTranspose weights. This is a temporary measurement but it will keep the weights flipped for now (including in the unit test txt files)."));
+				ConvTransposeOperator->SetWhetherWeightTensorWasFlipped(true);
+				TensorsIndex[1].Flip(2, TensorsIndex[1].GetNumberDimensions());
+			}
+			Operators.Emplace(ConvTransposeOperator);
+		}
+		// Remove all non-required/auxiliary tensors from InputTensors to keep it in X, W, and B (if B exists). Otherwise, those would be used by the network and moved into GPU memory
+		while (Tensors[TensorIndex].Num() > 4 || Tensors[TensorIndex].Last().Num() < 1)
+		{
+			Tensors[TensorIndex].Pop();
+		}
+		// TensorsIndex is no longer valid after this while loop!
+	}
+	TestOperator(InOutNetworkInferenceQAAsset, Tensors, InputTensors, OutputTensors, Operators, TEXT("FConvTransposeOperator"), InGroundTruthDirectory, InZeroThreshold,
+		/*bShouldOperatorsProvideSameResult*/false, /*bInShouldRunGPU*/true);
+}
+
+void FOperatorTester::TestReshapeOperator(UNeuralNetworkInferenceQAAsset* InOutNetworkInferenceQAAsset, const TMap<int32, TArray<float>>& InInputTensorData,
+	const TMap<int32, TArray<float>>& InInputTensorDataPos, const float InZeroThreshold, const FString& InGroundTruthDirectory, const int32 InOperatorCounter, const int32 InOperatorTotal)
+{
+	UE_LOG(LogNeuralNetworkInferenceQA, Display, TEXT("--------------- %d/%d FReshapeOperator"), InOperatorCounter, InOperatorTotal);
+	// Inout tensors (X != Mean to avoid Output = Bias; Variance positive to avoid NaN)
+	TArray<TArray<FNeuralTensor>> Tensors({
+		// ----------------------------------------------------------------- ONNX - [REDACTED_PROFILE_URL]/onnx/blob/master/docs/TestCoverage.md#Reshape -----------------------------------------------------------------
+		// reshape
+		{FNeuralTensor(InInputTensorData.FindChecked(24), TArray<int64>({2, 3, 4})), FNeuralTensor(TArray<int64>({4, 2, 3})), /*Output*/FNeuralTensor()},
+		{FNeuralTensor(InInputTensorData.FindChecked(24), TArray<int64>({2, 3, 4})), FNeuralTensor(TArray<int64>({2, 4, 3}))},
+		{FNeuralTensor(InInputTensorData.FindChecked(24), TArray<int64>({2, 3, 4})), FNeuralTensor(TArray<int64>({2, 12})), /*Output*/FNeuralTensor()},
+		{FNeuralTensor(InInputTensorData.FindChecked(24), TArray<int64>({2, 3, 4})), FNeuralTensor(TArray<int64>({2, 3, 2, 2}))},
+		{FNeuralTensor(InInputTensorData.FindChecked(24), TArray<int64>({2, 3, 4})), FNeuralTensor(TArray<int64>({24})), /*Output*/FNeuralTensor()},
+		{FNeuralTensor(InInputTensorData.FindChecked(24), TArray<int64>({2, 3, 4})), FNeuralTensor(TArray<int64>({2, -1, 2}))},
+		{FNeuralTensor(InInputTensorData.FindChecked(24), TArray<int64>({2, 3, 4})), FNeuralTensor(TArray<int64>({-1, 2, 3, 4})), /*Output*/FNeuralTensor()},
+		{FNeuralTensor(InInputTensorData.FindChecked(24), TArray<int64>({2, 3, 4})), FNeuralTensor(TArray<int64>({2, 0, 4, 1}))},
+		{FNeuralTensor(InInputTensorData.FindChecked(24), TArray<int64>({2, 3, 4})), FNeuralTensor(TArray<int64>({2, 0, 1, -1})), /*Output*/FNeuralTensor()},
+		// allowzero
+		{FNeuralTensor(FNeuralTensor(ENeuralDataType::Float, TArray<int64>({0, 3, 4}))), FNeuralTensor(TArray<int64>({3, 4, 0})), /*Output*/FNeuralTensor(), /*AllowZero=true*/FNeuralTensor()},
+		// Scalars
+		{FNeuralTensor(InInputTensorData.FindChecked(1)), FNeuralTensor(TArray<int64>({1})), /*Output*/FNeuralTensor()},
+		{FNeuralTensor(InInputTensorData.FindChecked(1)), FNeuralTensor(TArray<int64>({1}))},
+		// 1-D
+		{FNeuralTensor(InInputTensorData.FindChecked(10)), FNeuralTensor(TArray<int64>({5, 2})), /*Output*/FNeuralTensor()},
+		{FNeuralTensor(InInputTensorData.FindChecked(10)), FNeuralTensor(TArray<int64>({5, -1}))},
+		// 2-D
+		{FNeuralTensor(InInputTensorData.FindChecked(18), TArray<int64>({3, 6})), FNeuralTensor(TArray<int64>({-1, 3})), /*Output*/FNeuralTensor()},
+		{FNeuralTensor(InInputTensorData.FindChecked(18), TArray<int64>({3, 6})), FNeuralTensor(TArray<int64>({6, 3}))},
+		// 3-D
+		{FNeuralTensor(InInputTensorData.FindChecked(30), TArray<int64>({3, 5, 2})), FNeuralTensor(TArray<int64>({30})), /*Output*/FNeuralTensor()},
+		{FNeuralTensor(InInputTensorData.FindChecked(30), TArray<int64>({3, 5, 2})), FNeuralTensor(TArray<int64>({-1}))},
+		// 4-D
+		{FNeuralTensor(InInputTensorData.FindChecked(300), TArray<int64>({3, 10, 5, 2})), FNeuralTensor(TArray<int64>({-1, 10, 15})), /*Output*/FNeuralTensor()},
+		{FNeuralTensor(InInputTensorData.FindChecked(300), TArray<int64>({3, 10, 5, 2})), FNeuralTensor(TArray<int64>({2, 0, 15}))},
+		// 5-D
+		{FNeuralTensor(InInputTensorData.FindChecked(300), TArray<int64>({2, 3, 5, 5, 2})), FNeuralTensor(TArray<int64>({2, 10, 15})), /*Output*/FNeuralTensor()},
+		{FNeuralTensor(InInputTensorData.FindChecked(300), TArray<int64>({2, 3, 5, 5, 2})), FNeuralTensor(TArray<int64>({0, -1, 15}))}
+		});
+	TArray<TArray<FNeuralTensor*>> InputTensors;
+	TArray<TArray<FNeuralTensor*>> OutputTensors;
+	TArray<TSharedPtr<FNeuralOperator>> Operators;
+	for (int32 TensorIndex = 0; TensorIndex < Tensors.Num(); ++TensorIndex)
+	{
+		InputTensors.Push({ &Tensors[TensorIndex][0], &Tensors[TensorIndex][1] });
+		// Not inlined
+		if (Tensors[TensorIndex].Num() > 2)
+		{
+			OutputTensors.Push({ &Tensors[TensorIndex][2] });
+			Operators.Push(MakeShared<FReshapeOperator>(/*bIsInlinedTensor*/false, /*AllowZero*/(Tensors[TensorIndex].Num() > 3)));
+		}
+		// Inlined
+		else
+		{
+			OutputTensors.Push({});
+			Operators.Push(MakeShared<FReshapeOperator>(/*bIsInlinedTensor*/true, /*AllowZero*/0));
+		}
+	}
+	TestOperator(InOutNetworkInferenceQAAsset, Tensors, InputTensors, OutputTensors, Operators, TEXT("FReshapeOperator"), InGroundTruthDirectory, InZeroThreshold,
+		/*bShouldOperatorsProvideSameResult*/false, /*bInShouldRunGPU*/true);
+}
+
+void FOperatorTester::TestSqueezeOperator(UNeuralNetworkInferenceQAAsset* InOutNetworkInferenceQAAsset, const TMap<int32, TArray<float>>& InInputTensorData,
+	const TMap<int32, TArray<float>>& InInputTensorDataPos, const float InZeroThreshold, const FString& InGroundTruthDirectory, const int32 InOperatorCounter, const int32 InOperatorTotal)
+{
+	UE_LOG(LogNeuralNetworkInferenceQA, Display, TEXT("--------------- %d/%d FSqueezeOperator"), InOperatorCounter, InOperatorTotal);
+	// Inout tensors (X != Mean to avoid Output = Bias; Variance positive to avoid NaN)
+	TArray<TArray<FNeuralTensor>> Tensors({
+		// ----------------------------------------------------------------- ONNX - [REDACTED_PROFILE_URL]/onnx/blob/master/docs/Operators.md#Squeeze -----------------------------------------------------------------
+		// squeeze
+		{FNeuralTensor(InInputTensorData.FindChecked(60), TArray<int64>({1, 3, 4, 5})), FNeuralTensor(TArray<int64>({0})), /*Output*/FNeuralTensor()},
+		{FNeuralTensor(InInputTensorData.FindChecked(60), TArray<int64>({1, 3, 4, 5})), FNeuralTensor(TArray<int64>({0}))},
+		// squeeze_negative_axis
+		{FNeuralTensor(InInputTensorData.FindChecked(15), TArray<int64>({1, 3, 1, 5})), FNeuralTensor(TArray<int64>({-2})), /*Output*/FNeuralTensor()},
+		{FNeuralTensor(InInputTensorData.FindChecked(15), TArray<int64>({1, 3, 1, 5})), FNeuralTensor(TArray<int64>({-2}))},
+		// 1-D
+		{FNeuralTensor(InInputTensorData.FindChecked(5)), /*Empty*/FNeuralTensor(), /*Empty*/FNeuralTensor(), /*Output*/FNeuralTensor()},
+		// 2-D
+		{FNeuralTensor(InInputTensorData.FindChecked(60), TArray<int64>({12, 5})), /*Empty*/FNeuralTensor(), /*Empty*/FNeuralTensor(), /*Output*/FNeuralTensor()},
+		{FNeuralTensor(InInputTensorData.FindChecked(6), TArray<int64>({6, 1})), FNeuralTensor(TArray<int64>({1})), /*Output*/FNeuralTensor()},
+		{FNeuralTensor(InInputTensorData.FindChecked(6), TArray<int64>({6, 1})), FNeuralTensor(TArray<int64>({1}))},
+		{FNeuralTensor(InInputTensorData.FindChecked(6), TArray<int64>({6, 1}))},
+		{FNeuralTensor(InInputTensorData.FindChecked(6), TArray<int64>({6, 1})), /*Empty*/FNeuralTensor(), /*Empty*/FNeuralTensor(), /*Output*/FNeuralTensor()},
+		// 3-D
+		{FNeuralTensor(InInputTensorData.FindChecked(60), TArray<int64>({3, 4, 5})), /*Empty*/FNeuralTensor(), /*Empty*/FNeuralTensor(), /*Output*/FNeuralTensor()},
+		{FNeuralTensor(InInputTensorData.FindChecked(30), TArray<int64>({5, 1, 6})), FNeuralTensor(TArray<int64>({1}))},
+		{FNeuralTensor(InInputTensorData.FindChecked(30), TArray<int64>({5, 1, 6})), FNeuralTensor(TArray<int64>({-2})), /*Output*/FNeuralTensor()},
+		{FNeuralTensor(InInputTensorData.FindChecked(30), TArray<int64>({5, 1, 6})), FNeuralTensor(TArray<int64>({1})), /*Output*/FNeuralTensor()},
+		{FNeuralTensor(InInputTensorData.FindChecked(30), TArray<int64>({5, 1, 6})), /*Empty*/FNeuralTensor(), /*Empty*/FNeuralTensor(), /*Output*/FNeuralTensor()},
+		{FNeuralTensor(InInputTensorData.FindChecked(30), TArray<int64>({5, 1, 6}))},
+		// 4-D
+		{FNeuralTensor(InInputTensorData.FindChecked(60), TArray<int64>({3, 2, 2, 5})), /*Empty*/FNeuralTensor(), /*Empty*/FNeuralTensor(), /*Output*/FNeuralTensor()},
+		{FNeuralTensor(InInputTensorData.FindChecked(16), TArray<int64>({2, 8, 1, 1})), FNeuralTensor(TArray<int64>({-1, 2})), /*Output*/FNeuralTensor()},
+		{FNeuralTensor(InInputTensorData.FindChecked(16), TArray<int64>({2, 8, 1, 1})), /*Empty*/FNeuralTensor(), /*Empty*/FNeuralTensor(), /*Output*/FNeuralTensor()},
+		{FNeuralTensor(InInputTensorData.FindChecked(16), TArray<int64>({2, 8, 1, 1}))},
+		{FNeuralTensor(InInputTensorData.FindChecked(16), TArray<int64>({2, 8, 1, 1})), FNeuralTensor(TArray<int64>({2})), /*Output*/FNeuralTensor()},
+		// 5-D
+		{FNeuralTensor(InInputTensorData.FindChecked(50), TArray<int64>({1, 5, 1, 5, 2})), FNeuralTensor(TArray<int64>({0, 2})), /*Output*/FNeuralTensor()},
+		{FNeuralTensor(InInputTensorData.FindChecked(50), TArray<int64>({1, 5, 1, 5, 2})), FNeuralTensor(TArray<int64>({-3, -5})), /*Output*/FNeuralTensor()},
+		{FNeuralTensor(InInputTensorData.FindChecked(27), TArray<int64>({3, 1, 3, 1, 3})), FNeuralTensor(TArray<int64>({1, 3})), /*Output*/FNeuralTensor()},
+		{FNeuralTensor(InInputTensorData.FindChecked(27), TArray<int64>({3, 1, 3, 1, 3})), FNeuralTensor(TArray<int64>({3, 1})), /*Output*/FNeuralTensor()},
+		{FNeuralTensor(InInputTensorData.FindChecked(50), TArray<int64>({1, 5, 1, 5, 2}))}
+		//// Other tests (resulting in Warnings and/or unexpected behavior)
+		//// Scalars (undefined behavior: 0-d shape array for CPU and crash for GPU)
+		//{FNeuralTensor(InInputTensorData.FindChecked(1)), FNeuralTensor(TArray<int64>({0})), /*Output*/FNeuralTensor()},
+		//{FNeuralTensor(InInputTensorData.FindChecked(1)), FNeuralTensor(TArray<int64>({0}))},
+		//{FNeuralTensor(InInputTensorData.FindChecked(1)), /*Output*/FNeuralTensor()},
+		//{FNeuralTensor(InInputTensorData.FindChecked(1))},
+		//// Duplicated axis (generating warnings)
+		//{FNeuralTensor(InInputTensorData.FindChecked(16), TArray<int64>({2, 8, 1, 1})), FNeuralTensor(TArray<int64>({-2, 2})),  /*Output*/FNeuralTensor()},
+		//{FNeuralTensor(InInputTensorData.FindChecked(27), TArray<int64>({3, 1, 3, 1, 3})), FNeuralTensor(TArray<int64>({1, 1})), /*Output*/FNeuralTensor()},
+		//{FNeuralTensor(InInputTensorData.FindChecked(27), TArray<int64>({3, 1, 3, 1, 3})), FNeuralTensor(TArray<int64>({-4, 1})), /*Output*/FNeuralTensor()},
+		//{FNeuralTensor(InInputTensorData.FindChecked(27), TArray<int64>({3, 1, 3, 1, 3})), FNeuralTensor(TArray<int64>({1, -4})), /*Output*/FNeuralTensor()}
+		});
+	TArray<TArray<FNeuralTensor*>> InputTensors;
+	TArray<TArray<FNeuralTensor*>> OutputTensors;
+	TArray<TSharedPtr<FNeuralOperator>> Operators;
+	for (int32 TensorIndex = 0; TensorIndex < Tensors.Num(); ++TensorIndex)
+	{
+
+		// Two Input Tensors
+		if (Tensors[TensorIndex].Num() == 2 || Tensors[TensorIndex].Num() == 3)
+		{
+			InputTensors.Push({ &Tensors[TensorIndex][0], &Tensors[TensorIndex][1] });
+		}
+		// One Input Tensor
+		else
+		{
+			InputTensors.Push({ &Tensors[TensorIndex][0] });
+		}
+		
+		// Not inlined
+		if (Tensors[TensorIndex].Num() == 3)
+		{
+			OutputTensors.Push({ &Tensors[TensorIndex][2] });
+			Operators.Push(MakeShared<FSqueezeOperator>(/*bIsInlinedTensor*/false));
+		}
+		else if (Tensors[TensorIndex].Num() == 4)
+		{
+			OutputTensors.Push({ &Tensors[TensorIndex][3] });
+			Operators.Push(MakeShared<FSqueezeOperator>(/*bIsInlinedTensor*/false));
+		}
+		// Inlined
+		else
+		{
+			OutputTensors.Push({});
+			Operators.Push(MakeShared<FSqueezeOperator>(/*bIsInlinedTensor*/true));
+		}
+	}
+	TestOperator(InOutNetworkInferenceQAAsset, Tensors, InputTensors, OutputTensors, Operators, TEXT("FSqueezeOperator"), InGroundTruthDirectory, InZeroThreshold,
+		/*bShouldOperatorsProvideSameResult*/false, /*bInShouldRunGPU*/true);
+}
+
+/* FOperatorTester static private functions
+ *****************************************************************************/
+
+TArray<int64> FOperatorTester::FloatTensorToIntArray(const TArray<FNeuralTensor>& InTensors, const int32 InTensorIndex)
+{
+	// Set OutputArray - Hacky but simple way of getting dilation values
+	TArray<int64> OutputArray;
+	if (InTensors.Num() > InTensorIndex && InTensors[InTensorIndex].Num() > 0)
+	{
+		const FNeuralTensor& Tensor = InTensors[InTensorIndex];
+		OutputArray.SetNumUninitialized(Tensor.Num());
+		for (int32 Index = 0; Index < Tensor.Num(); ++Index)
+		{
+			OutputArray[Index] = FMath::RoundToInt(Tensor.At<float>(Index));
+		}
+	}
+	return OutputArray;
+}
+
+void FOperatorTester::ResetOrCheckInput(UNeuralNetwork* InOutNetwork, const TArray<FNeuralTensor>& InInputTensorArray, const bool bIsInlinedTensor)
+{
+	if (bIsInlinedTensor)
+	{
+		InOutNetwork->SetInputFromArrayCopy(InInputTensorArray);
+	}
+	// Check input has not changed
+	else
+	{
+		for (int32 InputIndex = 0; InputIndex < InOutNetwork->GetInputTensorNumber(); ++InputIndex)
+		{
+			const FNeuralTensor& InputTensor = InOutNetwork->GetInputTensor(InputIndex);
+			ensureMsgf(InInputTensorArray[InputIndex].Num() == InputTensor.Num(), TEXT("InInputTensorArray[InputIndex].Num() == InputTensor.Num() failed (%d != %d)."),
+				InInputTensorArray[InputIndex].Num(), InputTensor.Num());
+			ensureMsgf(FNeuralNetworkInferenceQAUtils::EstimateTensorL1DiffError(InputTensor, InInputTensorArray[InputIndex], 0.f, TEXT("ResetOrCheckInput")), TEXT("ResetOrCheckInput() did not pass."));
+		}
+	}
+}
+
+void FOperatorTester::NetworkGPUvsCPU(UNeuralNetwork* InOutNetwork, const bool bIsInlinedTensor, const float InZeroThreshold)
+{
+	// Create InputTensorDataArray
+	const TArray<FNeuralTensor> InputTensorArray = InOutNetwork->CreateInputArrayCopy();
+	// GPU
+	InOutNetwork->SetDeviceType(ENeuralDeviceType::GPU);
+	InOutNetwork->Run();
+	// Store GPU results
+	const TArray<FNeuralTensor> OutputTensorArray = InOutNetwork->CreateOutputArrayCopy();
+	// Reset memory / Check input has not changed
+	ResetOrCheckInput(InOutNetwork, InputTensorArray, bIsInlinedTensor);
+	// CPU
+	InOutNetwork->SetDeviceType(ENeuralDeviceType::CPU);
+	InOutNetwork->Run();
+	// CPU vs. GPU
+	for (int32 OutputIndex = 0; OutputIndex < InOutNetwork->GetOutputTensorNumber(); ++OutputIndex)
+	{
+		const FNeuralTensor& OutputTensorCPU = OutputTensorArray[OutputIndex];
+		const FNeuralTensor& OutputTensorGPU = InOutNetwork->GetOutputTensor(OutputIndex);
+		ensureMsgf(OutputTensorCPU.Num() == OutputTensorGPU.Num(), TEXT("OutputTensorCPU.Num() == OutputTensorGPU.Num(): %d vs. %d"),
+			OutputTensorCPU.Num(), OutputTensorGPU.Num());
+		ensureMsgf(FNeuralNetworkInferenceQAUtils::EstimateTensorL1DiffError(OutputTensorCPU, OutputTensorGPU, InZeroThreshold, TEXT("GPUvsCPU")), TEXT("CPU vs. GPU error has too high"));
+	}
+	// Reset memory / Check input has not changed
+	ResetOrCheckInput(InOutNetwork, InputTensorArray, bIsInlinedTensor); // This is doing an unnecessary copy (it should move really)
+}
+
+void FOperatorTester::NetworkInputVsOutput(UNeuralNetwork* InOutNetwork, const int32 InInlinedTensorIndex)
+{
+	const TArray<FNeuralTensor> InputTensorArray = InOutNetwork->CreateInputArrayCopy();
+	// Run UNeuralNetwork
+	InOutNetwork->Run();
+	// Check input equals one of the outputs
+	const FNeuralTensor& InputTensor = InOutNetwork->GetInputTensor(InInlinedTensorIndex);
+	bool bIsSomeInputAnOutput = false;
+	for (int32 OutputIndex = 0; OutputIndex < InOutNetwork->GetOutputTensorNumber(); ++OutputIndex)
+	{
+		const FNeuralTensor& OutputTensor = InOutNetwork->GetOutputTensor(OutputIndex);
+		if (InputTensor.GetUnderlyingUInt8ArrayRef() == OutputTensor.GetUnderlyingUInt8ArrayRef())
+		{
+			bIsSomeInputAnOutput = true;
+			break;
+		}
+	}
+	ensureMsgf(bIsSomeInputAnOutput, TEXT("Input is not being used as an output."));
+	// Reset memory
+	InOutNetwork->SetInputFromArrayCopy(InputTensorArray); // This is doing an unnecessary copy (it should move really)
+}
diff --git a/Engine/Plugins/Experimental/NeuralNetworkInference/Source/Private/NeuralNetworkInferenceQA/Private/OperatorTester.h b/Engine/Plugins/Experimental/NeuralNetworkInference/Source/Private/NeuralNetworkInferenceQA/Private/OperatorTester.h
new file mode 100644
--- /dev/null
+++ b/Engine/Plugins/Experimental/NeuralNetworkInference/Source/Private/NeuralNetworkInferenceQA/Private/OperatorTester.h
@@ -0,0 +1,149 @@
+// Copyright Epic Games, Inc. All Rights Reserved.
+
+#pragma once
+
+#include "CoreMinimal.h"
+#include "NeuralEnumClasses.h"
+#include "NeuralNetwork.h"
+#include "NeuralNetworkInferenceQAUtils.h"
+#include "NeuralNetworkInferenceQAAsset.h"
+#include "NeuralOperator.h"
+#include "NeuralTensor.h"
+
+class FOperatorTester
+{
+public:
+	/**
+	 * @param InOutOutputTensors It must be empty if the operator is inlined.
+	 * @param bInShouldOperatorsProvideSameResult If true, each element of InOutOperators is supposed to give the same exact result (so it will be stored with the same name).
+	 * @param bInShouldRunFinalStringVsIOTest If true, the FString InOutAllNetworkStrings will be compared against the GT file. If false, the FString will just be returned
+	 * to the caller. The caller can stack and accumulate the FString and compare it against a global GT file (benefit: 1 file per layer).
+	 */
+	static void TestOperator(UNeuralNetworkInferenceQAAsset* InOutNetworkInferenceQAAsset, TArray<TArray<FNeuralTensor>>& InOutTensors, TArray<TArray<FNeuralTensor*>>& InOutInputTensors,
+		TArray<TArray<FNeuralTensor*>>& InOutOutputTensors, TArray<TSharedPtr<FNeuralOperator>>& InOutOperators, const FString& InOperatorName,
+		const FString& InGroundTruthDirectory, const float InZeroThreshold, const bool bInShouldOperatorsProvideSameResult,
+		const bool bInShouldRunGPU = true, const bool bInShouldRunFinalStringVsIOTest = true);
+
+	static void TestBatchNormalizationOperator(UNeuralNetworkInferenceQAAsset* InOutNetworkInferenceQAAsset, const TMap<int32, TArray<float>>& InInputTensorData, const TMap<int32, TArray<float>>& InInputTensorDataSmall,
+		const TMap<int32, TArray<float>>& InInputTensorDataPos, const float InZeroThreshold, const FString& InGroundTruthDirectory, const int32 InOperatorCounter, const int32 InOperatorTotal);
+
+	static void TestGemmOperator(UNeuralNetworkInferenceQAAsset* InOutNetworkInferenceQAAsset, const TMap<int32, TArray<float>>& InInputTensorData, const float InZeroThreshold, const FString& InGroundTruthDirectory,
+		const int32 InOperatorCounter, const int32 InOperatorTotal);
+
+	static void TestConvOperator(UNeuralNetworkInferenceQAAsset* InOutNetworkInferenceQAAsset, const TMap<int32, TArray<float>>& InInputTensorData, const TMap<int32, TArray<float>>& InInputTensorDataSmall,
+		const float InZeroThreshold, const FString& InGroundTruthDirectory, const int32 InOperatorCounter, const int32 InOperatorTotal);
+
+	static void TestConvTransposeOperator(UNeuralNetworkInferenceQAAsset* InOutNetworkInferenceQAAsset, const TMap<int32, TArray<float>>& InInputTensorData, const TMap<int32, TArray<float>>& InInputTensorDataSmall,
+		const float InZeroThreshold, const FString& InGroundTruthDirectory, const int32 InOperatorCounter, const int32 InOperatorTotal);
+
+	template<typename TNeuralOperatorChild>
+	static void TestElementWiseOperator(UNeuralNetworkInferenceQAAsset* InOutNetworkInferenceQAAsset, const TCHAR* InOperatorName, const TArray<TArray<FNeuralTensor>>& InTensors, const float InZeroThreshold,
+		const FString& InGroundTruthDirectory, const int32 InOperatorCounter, const int32 InOperatorTotal);
+
+	template<typename TNeuralOperatorChild>
+	static void TestElementWise1AttributeOperator(UNeuralNetworkInferenceQAAsset* InOutNetworkInferenceQAAsset, const TCHAR* InOperatorName, const TArray<TArray<FNeuralTensor>>& InTensors, const float InZeroThreshold,
+		const FString& InGroundTruthDirectory, const int32 InOperatorCounter, const int32 InOperatorTotal);
+
+	template<typename TNeuralOperatorChild>
+	static void TestMultiDirectionalBroadcastingOperator(UNeuralNetworkInferenceQAAsset* InOutNetworkInferenceQAAsset, const TCHAR* InOperatorName, const TArray<TArray<FNeuralTensor>>& InTensors, const float InZeroThreshold,
+		const FString& InGroundTruthDirectory, const int32 InOperatorCounter, const int32 InOperatorTotal);
+
+	static void TestReshapeOperator(UNeuralNetworkInferenceQAAsset* InOutNetworkInferenceQAAsset, const TMap<int32, TArray<float>>& InInputTensorData,
+		const TMap<int32, TArray<float>>& InInputTensorDataPos, const float InZeroThreshold, const FString& InGroundTruthDirectory, const int32 InOperatorCounter, const int32 InOperatorTotal);
+
+	static void TestSqueezeOperator(UNeuralNetworkInferenceQAAsset* InOutNetworkInferenceQAAsset, const TMap<int32, TArray<float>>& InInputTensorData,
+		const TMap<int32, TArray<float>>& InInputTensorDataPos, const float InZeroThreshold, const FString& InGroundTruthDirectory, const int32 InOperatorCounter, const int32 InOperatorTotal);
+
+private:
+	static TArray<int64> FloatTensorToIntArray(const TArray<FNeuralTensor>& InTensors, const int32 InTensorIndex);
+
+	/**
+	 * For inlined layers, it resets the memory (meant after a UNeuralNetwork::Run()). For non-inlined layers, it checks that the input did not change.
+	 */
+	static void ResetOrCheckInput(UNeuralNetwork* InOutNetwork, const TArray<FNeuralTensor>& InInputTensorArray, const bool bIsInlinedTensor);
+
+	/**
+	 * It tests and makes sure CPU and GPU modes provide the same results. Otherwise, it will check to false.
+	 * @param InOutNetwork UNeuralNetwork to test.
+	 * @param If not inlined, it will make check whether the input tensors have not changed after running.
+	 * @param InZeroThreshold If the average L2 norm between CPU and GPU results is greater than this number, it will check to false.
+	 */
+	static void NetworkGPUvsCPU(UNeuralNetwork* InOutNetwork, const bool bIsInlinedTensor, const float InZeroThreshold = 0.f);
+
+	/**
+	 * If layer inlined --> Check input == output
+	 * This is only meant for networks with single input single output, whose layers are all inlined (so that input tensor data should match
+	 * output tensor data after a UNeuralNetwork run).
+	 */
+	static void NetworkInputVsOutput(UNeuralNetwork* InOutNetwork, const int32 InInlinedTensorIndex);
+};
+
+
+
+/* FOperatorTester inlined and templated functions
+ *****************************************************************************/
+
+template<typename TNeuralOperatorChild>
+void FOperatorTester::TestElementWiseOperator(UNeuralNetworkInferenceQAAsset* InOutNetworkInferenceQAAsset, const TCHAR* InOperatorName, const TArray<TArray<FNeuralTensor>>& InTensors, const float InZeroThreshold,
+	const FString& InGroundTruthDirectory, const int32 InOperatorCounter, const int32 InOperatorTotal)
+{
+	UE_LOG(LogNeuralNetworkInferenceQA, Display, TEXT("--------------- %d/%d %s"), InOperatorCounter, InOperatorTotal, InOperatorName);
+	TArray<TArray<FNeuralTensor>> TensorsCopy = InTensors;
+	TArray<TArray<FNeuralTensor*>> InputTensors({ {&TensorsCopy[0][0]}, {&TensorsCopy[1][0]}, {&TensorsCopy[2][0]}, {&TensorsCopy[3][0]}, {&TensorsCopy[4][0]} });
+	TArray<TArray<FNeuralTensor*>> OutputTensors({ TArray<FNeuralTensor*>({/*Inlined*/}), {&TensorsCopy[1][1]}, TArray<FNeuralTensor*>({/*Inlined*/}), {&TensorsCopy[3][1]},
+		TArray<FNeuralTensor*>({/*Inlined*/}) });
+	TArray<TSharedPtr<FNeuralOperator>> Operators({ MakeShared<TNeuralOperatorChild>(/*bIsInlinedTensor*/true), MakeShared<TNeuralOperatorChild>(false), MakeShared<TNeuralOperatorChild>(true),
+		MakeShared<TNeuralOperatorChild>(false), MakeShared<TNeuralOperatorChild>(true) });
+	TestOperator(InOutNetworkInferenceQAAsset, TensorsCopy, InputTensors, OutputTensors, Operators, InOperatorName, InGroundTruthDirectory, InZeroThreshold, /*bShouldOperatorsProvideSameResult*/false,
+		/*bInShouldRunGPU*/true);
+}
+
+template<typename TNeuralOperatorChild>
+void FOperatorTester::TestElementWise1AttributeOperator(UNeuralNetworkInferenceQAAsset* InOutNetworkInferenceQAAsset, const TCHAR* InOperatorName, const TArray<TArray<FNeuralTensor>>& InTensors, const float InZeroThreshold,
+	const FString& InGroundTruthDirectory, const int32 InOperatorCounter, const int32 InOperatorTotal)
+{
+	UE_LOG(LogNeuralNetworkInferenceQA, Display, TEXT("--------------- %d/%d %s"), InOperatorCounter, InOperatorTotal, InOperatorName);
+	TArray<TArray<FNeuralTensor>> TensorsCopy = InTensors;
+	TensorsCopy += InTensors; // Append not implemented against itself
+	TArray<TArray<FNeuralTensor*>> InputTensors({ {&TensorsCopy[0][0]}, {&TensorsCopy[1][0]}, {&TensorsCopy[2][0]}, {&TensorsCopy[3][0]}, {&TensorsCopy[4][0]},
+		{&TensorsCopy[5][0]}, {&TensorsCopy[6][0]}, {&TensorsCopy[7][0]}, {&TensorsCopy[8][0]}, {&TensorsCopy[9][0]} });
+	TArray<TArray<FNeuralTensor*>> OutputTensors({ TArray<FNeuralTensor*>({/*Inlined*/}), {&TensorsCopy[1][1]}, TArray<FNeuralTensor*>({/*Inlined*/}), {&TensorsCopy[3][1]},
+		TArray<FNeuralTensor*>({/*Inlined*/}), TArray<FNeuralTensor*>({/*Inlined*/}), {&TensorsCopy[6][1]}, TArray<FNeuralTensor*>({/*Inlined*/}), {&TensorsCopy[8][1]},
+		TArray<FNeuralTensor*>({/*Inlined*/}) });
+	TArray<TSharedPtr<FNeuralOperator>> Operators({ MakeShared<TNeuralOperatorChild>(/*bIsInlinedTensor*/true, 1.f), MakeShared<TNeuralOperatorChild>(false, 20.f), MakeShared<TNeuralOperatorChild>(true, 0.5f),
+		MakeShared<TNeuralOperatorChild>(false, 0.01f), MakeShared<TNeuralOperatorChild>(true, 0.f), MakeShared<TNeuralOperatorChild>(/*bIsInlinedTensor*/true, -1.f), MakeShared<TNeuralOperatorChild>(false, -20.f),
+		MakeShared<TNeuralOperatorChild>(true, -0.5f), MakeShared<TNeuralOperatorChild>(false, -0.01f), MakeShared<TNeuralOperatorChild>(true, 1.2345f) });
+	TestOperator(InOutNetworkInferenceQAAsset, TensorsCopy, InputTensors, OutputTensors, Operators, InOperatorName, InGroundTruthDirectory, InZeroThreshold, /*bShouldOperatorsProvideSameResult*/false,
+		/*bInShouldRunGPU*/true);
+}
+
+template<typename TNeuralOperatorChild>
+void FOperatorTester::TestMultiDirectionalBroadcastingOperator(UNeuralNetworkInferenceQAAsset* InOutNetworkInferenceQAAsset, const TCHAR* InOperatorName, const TArray<TArray<FNeuralTensor>>& InTensors, const float InZeroThreshold,
+	const FString& InGroundTruthDirectory, const int32 InOperatorCounter, const int32 InOperatorTotal)
+{
+	UE_LOG(LogNeuralNetworkInferenceQA, Display, TEXT("--------------- %d/%d %s"), InOperatorCounter, InOperatorTotal, InOperatorName);
+
+	const TSet<uint32> SetOperatorsAllInlined({ 0u, 1u });
+	const TSet<uint32> SetOperatorAInlined({ 0u });
+	const TSet<uint32> SetOperatorBInlined({ 1u });
+	const TSet<uint32> SetNoOperatorsInlined;
+	for (int32 TensorTestsIndex = 0; TensorTestsIndex < InTensors.Num(); ++TensorTestsIndex)
+	{
+		UE_LOG(LogNeuralNetworkInferenceQA, Display, TEXT("-------------------- %s #%d/%d"), InOperatorName, TensorTestsIndex + 1, InTensors.Num());
+		TArray<TSharedPtr<FNeuralOperator>> Operators({ MakeShared<TNeuralOperatorChild>(SetOperatorsAllInlined), MakeShared<TNeuralOperatorChild>(SetOperatorAInlined),
+			MakeShared<TNeuralOperatorChild>(SetOperatorBInlined), MakeShared<TNeuralOperatorChild>(SetNoOperatorsInlined) });
+		// Input/output tensors
+		TArray<TArray<FNeuralTensor>> MultiBroadcastTensors;
+		MultiBroadcastTensors.Init(InTensors[TensorTestsIndex], Operators.Num());
+		TArray<TArray<FNeuralTensor*>> InputTensors;
+		TArray<TArray<FNeuralTensor*>> OutputTensors;
+		for (int32 OperatorIndex = 0; OperatorIndex < Operators.Num(); ++OperatorIndex)
+		{
+			InputTensors.Push({ &MultiBroadcastTensors[OperatorIndex][0], &MultiBroadcastTensors[OperatorIndex][1] });
+			OutputTensors.Push({ &MultiBroadcastTensors[OperatorIndex][2] });
+		}
+		// Run actual layer unit testing
+		FOperatorTester::TestOperator(InOutNetworkInferenceQAAsset, MultiBroadcastTensors, InputTensors, OutputTensors, Operators, InOperatorName, InGroundTruthDirectory, InZeroThreshold,
+			/*bShouldOperatorsProvideSameResult*/true, /*bShouldRunGPU*/true, (TensorTestsIndex + 1 == InTensors.Num()));
+	}
+}
diff --git a/Engine/Plugins/Experimental/NeuralNetworkInference/Source/Private/NeuralNetworkInferenceQA/Private/OperatorUnitTester.cpp b/Engine/Plugins/Experimental/NeuralNetworkInference/Source/Private/NeuralNetworkInferenceQA/Private/OperatorUnitTester.cpp
new file mode 100644
--- /dev/null
+++ b/Engine/Plugins/Experimental/NeuralNetworkInference/Source/Private/NeuralNetworkInferenceQA/Private/OperatorUnitTester.cpp
@@ -0,0 +1,255 @@
+// Copyright Epic Games, Inc. All Rights Reserved.
+
+#include "OperatorUnitTester.h"
+#include "NeuralEnumClasses.h"
+#include "NeuralNetworkInferenceQAUtils.h"
+#include "NeuralOperator.h"
+#include "NeuralOperators.h"
+#include "NeuralTensor.h"
+#include "OperatorTester.h"
+#include "Math/RandomStream.h"
+
+
+
+/* FOperatorUnitTester static public functions
+ *****************************************************************************/
+
+bool FOperatorUnitTester::GlobalTest(const FString& InProjectContentDir, const FString& InUnitTestRelativeDirectory)
+{
+	const FString UnitTestDirectory = InProjectContentDir / InUnitTestRelativeDirectory;
+	const float ZeroThreshold = 0.f;
+	const float FloatThreshold = 5e-7;
+	const float RelaxedFloatThreshold = 5e-6;
+	const float VeryRelaxedFloatThreshold = 4e-5;
+	const float VeryVeryRelaxedFloatThreshold = 2e-4;
+	const float SuperRelaxedFloatThreshold = 1.5e-3;
+	// Input arrays
+	const TArray<int32> ArraySizesToPopulate({
+		1,
+		2,
+		3,
+		4, // 2^2
+		5,
+		6, // 2x3
+		8, // 2^3
+		9, // 3^2
+		10, // 2x5
+		12, // 2^2 x 3
+		15, // 3x5
+		16, // 2^4
+		18, // 2x3x3
+		20, // 2^2 x 5
+		21, // 3x7
+		23,
+		24, // 2^3 x 3
+		25, // 5^2
+		27, // 3^3
+		30, // 2 x 3 x 5
+		32, // 2^5
+		36, // 2^2 x 3^2
+		40, // 2^3 x 5
+		46, // 2 x 23
+		50, // 2 x 5^2
+		60, // 1 x 2^2 x 3 x 5
+		64, // 2^6
+		69, // 3 x 23
+		72, // 2^3 x 3^2
+		80, // 2^4 x 5
+		81, // 3^4
+		90, // 2 x 3^2 x 5
+		100, // 2^2 x 5^2
+		108, // 2^2 x 3^3
+		112, // 2^4 x 7
+		120, // 2^3 x 3 x 5
+		125, // 5^3
+		128, // 2^7
+		150, // 2 x 3 x 5^2
+		160, // 2^5 x 5
+		162, // 2 x 3^4
+		240, // 2^4 x 3 x 5
+		243, // 3^5
+		256, // 2^8
+		300, // 2^2 x 3 x 5^2
+		486, // 2 x 3^5
+		512, // 2^9
+		625, // 5^4
+		810, // 2 x 3^4 x 5
+		1800, // 2^3 x 3^2 x 5^21
+		2700 // 2^2 x 3^3 x 5^2
+		// To be sorted out properly later
+		// [Add new numbers here, DO NOT add above this line]
+	});
+	TMap<int32, TArray<float>> InputTensorData;
+	TMap<int32, TArray<float>> InputTensorDataSmall;
+	TMap<int32, TArray<float>> InputTensorDataTiny;
+	TMap<int32, TArray<float>> InputTensorDataPos;
+	CreateRandomArraysOfFixedSeedAndDefinedSizes(InputTensorData, InputTensorDataSmall, InputTensorDataTiny, InputTensorDataPos, ArraySizesToPopulate);
+	// Input tensors
+	const TArray<int64> Sizes4D({ 2, 3, 5, 5 });
+	const FNeuralTensor InputTensor0D(InputTensorData.FindChecked(1));
+	const FNeuralTensor InputTensor1D_150(InputTensorData.FindChecked(150));
+	const FNeuralTensor InputTensor2D_18(InputTensorData.FindChecked(18), TArray<int64>({ 1, 18 }));
+	const FNeuralTensor InputTensor3D_300(InputTensorData.FindChecked(300), TArray<int64>({ 5, 6, 10 }));
+	const FNeuralTensor InputTensor4D_150(InputTensorData.FindChecked(150), TArray<int64>({ 2, 3, 5, 5 }));
+	const FNeuralTensor InputTensorSmall0D(InputTensorDataSmall.FindChecked(1));
+	const FNeuralTensor InputTensorSmall1D_150(InputTensorDataSmall.FindChecked(150));
+	const FNeuralTensor InputTensorSmall2D_18(InputTensorDataSmall.FindChecked(18), InputTensor2D_18.GetSizes());
+	const FNeuralTensor InputTensorSmall3D_300(InputTensorDataSmall.FindChecked(300), InputTensor3D_300.GetSizes());
+	const FNeuralTensor InputTensorSmall4D_150(InputTensorDataSmall.FindChecked(150), InputTensor4D_150.GetSizes());
+	const FNeuralTensor InputTensorTiny0D(InputTensorDataTiny.FindChecked(1), InputTensor0D.GetSizes());
+	const FNeuralTensor InputTensorTiny1D_150(InputTensorDataTiny.FindChecked(150));
+	const FNeuralTensor InputTensorTiny2D_18(InputTensorDataTiny.FindChecked(18), InputTensor2D_18.GetSizes());
+	const FNeuralTensor InputTensorTiny3D_300(InputTensorDataTiny.FindChecked(300), InputTensor3D_300.GetSizes());
+	const FNeuralTensor InputTensorTiny4D_150(InputTensorDataTiny.FindChecked(150), InputTensor4D_150.GetSizes());
+	// Inout tensors for Elementwise Operators
+	const TArray<TArray<FNeuralTensor>> ElementWiseTensors({ {InputTensor0D}, {InputTensor1D_150, InputTensor1D_150}, {InputTensor2D_18},
+		{InputTensor3D_300, InputTensor3D_300}, {InputTensor4D_150} });
+	const TArray<TArray<FNeuralTensor>> ElementWiseTensorsSmall({ {InputTensorSmall0D}, {InputTensorSmall1D_150, InputTensorSmall1D_150},
+		{InputTensorSmall2D_18}, {InputTensorSmall3D_300, InputTensorSmall3D_300}, {InputTensorSmall4D_150} });
+	const TArray<TArray<FNeuralTensor>> ElementWiseTensorsTiny({ {InputTensorTiny0D}, {InputTensorTiny1D_150, InputTensorTiny1D_150},
+		{InputTensorTiny2D_18}, {InputTensorTiny3D_300, InputTensorTiny3D_300}, {InputTensorTiny4D_150} });
+
+	// Inout tensors for Multidirectional Broadcasting Operators
+	const TArray<int64> Sizes7D({ 2, 3, 3, 2, 5, 3, 5 });
+	const TArray<TArray<FNeuralTensor>> MultidirectionalTensors({
+		// Scalars
+		{InputTensor0D, InputTensorSmall0D, FNeuralTensor(ENeuralDataType::Float, InputTensor0D.GetSizes())},
+		// 1-D, 2-D, 4-D, 7-D
+		{FNeuralTensor(InputTensorData.FindChecked(150)), FNeuralTensor(InputTensorDataSmall.FindChecked(150)), FNeuralTensor(ENeuralDataType::Float, 150)},
+		{InputTensor2D_18, InputTensorSmall2D_18, FNeuralTensor(ENeuralDataType::Float, InputTensor2D_18.GetSizes())},
+		{InputTensor4D_150, InputTensorSmall4D_150, FNeuralTensor(ENeuralDataType::Float, InputTensor4D_150.GetSizes())},
+		{FNeuralTensor(InputTensorData.FindChecked(2700), Sizes7D), FNeuralTensor(InputTensorDataSmall.FindChecked(2700), Sizes7D),
+			FNeuralTensor(ENeuralDataType::Float, Sizes7D)},
+		// Multidirectional Broadcasting tests - One of the inputs is 1D
+		{InputTensor4D_150, InputTensorSmall0D, FNeuralTensor(ENeuralDataType::Float, InputTensor4D_150.GetSizes())},
+		{InputTensor0D, InputTensorSmall2D_18, FNeuralTensor(ENeuralDataType::Float, InputTensor2D_18.GetSizes())},
+		// Multidirectional Broadcasting tests - Random input sizes
+		{FNeuralTensor(InputTensorData.FindChecked(3)), FNeuralTensor(InputTensorDataSmall.FindChecked(3), TArray<int64>({ 3, 1 })),
+			FNeuralTensor(ENeuralDataType::Float, TArray<int64>{3, 3})}, // [3].^[3, 1] --> [3,3]
+		{FNeuralTensor(InputTensorData.FindChecked(3), TArray<int64>({ 3, 1 })), FNeuralTensor(InputTensorDataSmall.FindChecked(3)),
+			FNeuralTensor(ENeuralDataType::Float, TArray<int64>{3, 3})}, // [3, 1].^[3] --> [3,3]
+		// {2, 1, 2, 1, 3, 5, 5} .^ {2, 2, 5, 3, 1, 5} --> {2, 2, 2, 5, 3, 5, 5}
+		{FNeuralTensor(InputTensorData.FindChecked(300), TArray<int64>({ 2, 1, 2, 1, 3, 5, 5 })),
+			FNeuralTensor(InputTensorDataSmall.FindChecked(300), TArray<int64>({ 2, 2, 5, 3, 1, 5 })),
+			FNeuralTensor(ENeuralDataType::Float, TArray<int64>{2, 2, 2, 5, 3, 5, 5})}
+		});
+	// Get all operator names
+	const TArray<FString> OperatorNames({ TEXT("FAbsOperator"), TEXT("FAcosOperator"), TEXT("FAddOperator"), TEXT("FAsinOperator"), TEXT("FAtanOperator"), TEXT("FBatchNormalizationOperator"), TEXT("FCeilOperator"),
+		TEXT("FConvOperator"), TEXT("FConvTransposeOperator"), TEXT("FCosOperator"), TEXT("FCoshOperator"), TEXT("FDivOperator"), TEXT("FExpOperator"), TEXT("FFloorOperator"), TEXT("FGemmOperator"),
+		TEXT("FLeakyReluOperator"), TEXT("FLogOperator"), TEXT("FMulOperator"), TEXT("FNegOperator"), TEXT("FPowOperator"), TEXT("FReciprocalOperator"), TEXT("FReluOperator"), TEXT("FReshapeOperator"),
+		TEXT("FRoundOperator"), TEXT("FSigmoidOperator"), TEXT("FSignOperator"), TEXT("FSinOperator"), TEXT("FSinhOperator"), TEXT("FSqrtOperator"), TEXT("FSqueezeOperator"), TEXT("FSubOperator"), 
+		TEXT("FTanOperator"), TEXT("FTanhOperator") });
+	// Creating NeuralNetworkInferenceQAAsset
+	UNeuralNetworkInferenceQAAsset* NeuralNetworkInferenceQAAsset = UNeuralNetworkInferenceQAAsset::Load(InUnitTestRelativeDirectory, TEXT("NeuralNetworkInferenceQAAsset"));
+	if (!NeuralNetworkInferenceQAAsset)
+	{
+		ensureMsgf(false, TEXT("NeuralNetworkInferenceQAAsset could not be loaded!"));
+		return false;
+	}
+	NeuralNetworkInferenceQAAsset->FindOrAddOperators(OperatorNames);
+	NeuralNetworkInferenceQAAsset->FlushNewTests();
+	// Per operator testing (checking false if errors occur)
+	const int32 OperatorTotal = OperatorNames.Num(); // 164;
+	int32 OperatorCounter = 0;
+	FOperatorTester::TestElementWiseOperator<FAbsOperator>(NeuralNetworkInferenceQAAsset, TEXT("FAbsOperator"), ElementWiseTensors, ZeroThreshold, UnitTestDirectory, ++OperatorCounter, OperatorTotal);
+	FOperatorTester::TestElementWiseOperator<FAcosOperator>(NeuralNetworkInferenceQAAsset, TEXT("FAcosOperator"), ElementWiseTensorsTiny, VeryVeryRelaxedFloatThreshold, UnitTestDirectory, ++OperatorCounter, OperatorTotal);
+	FOperatorTester::TestMultiDirectionalBroadcastingOperator<FAddOperator>(NeuralNetworkInferenceQAAsset, TEXT("FAddOperator"), MultidirectionalTensors, ZeroThreshold, UnitTestDirectory, ++OperatorCounter, OperatorTotal);
+	FOperatorTester::TestElementWiseOperator<FAsinOperator>(NeuralNetworkInferenceQAAsset, TEXT("FAsinOperator"), ElementWiseTensorsTiny, VeryVeryRelaxedFloatThreshold, UnitTestDirectory, ++OperatorCounter, OperatorTotal);
+	FOperatorTester::TestElementWiseOperator<FAtanOperator>(NeuralNetworkInferenceQAAsset, TEXT("FAtanOperator"), ElementWiseTensors, VeryVeryRelaxedFloatThreshold, UnitTestDirectory, ++OperatorCounter, OperatorTotal);
+	FOperatorTester::TestBatchNormalizationOperator(NeuralNetworkInferenceQAAsset, InputTensorData, InputTensorDataSmall, InputTensorDataPos, FloatThreshold, UnitTestDirectory, ++OperatorCounter, OperatorTotal);
+	FOperatorTester::TestElementWiseOperator<FCeilOperator>(NeuralNetworkInferenceQAAsset, TEXT("FCeilOperator"), ElementWiseTensors, ZeroThreshold, UnitTestDirectory, ++OperatorCounter, OperatorTotal);
+	FOperatorTester::TestConvOperator(NeuralNetworkInferenceQAAsset, InputTensorData, InputTensorDataSmall, FloatThreshold, UnitTestDirectory, ++OperatorCounter, OperatorTotal);
+	FOperatorTester::TestConvTransposeOperator(NeuralNetworkInferenceQAAsset, InputTensorData, InputTensorDataSmall, FloatThreshold, UnitTestDirectory, ++OperatorCounter, OperatorTotal);
+	FOperatorTester::TestElementWiseOperator<FCosOperator>(NeuralNetworkInferenceQAAsset, TEXT("FCosOperator"), ElementWiseTensors, VeryRelaxedFloatThreshold, UnitTestDirectory, ++OperatorCounter, OperatorTotal);
+	FOperatorTester::TestElementWiseOperator<FCoshOperator>(NeuralNetworkInferenceQAAsset, TEXT("FCoshOperator"), ElementWiseTensorsSmall, FloatThreshold, UnitTestDirectory, ++OperatorCounter, OperatorTotal);
+	FOperatorTester::TestMultiDirectionalBroadcastingOperator<FDivOperator>(NeuralNetworkInferenceQAAsset, TEXT("FDivOperator"), MultidirectionalTensors, FloatThreshold, UnitTestDirectory, ++OperatorCounter, OperatorTotal);
+	FOperatorTester::TestElementWiseOperator<FExpOperator>(NeuralNetworkInferenceQAAsset, TEXT("FExpOperator"), ElementWiseTensorsSmall, FloatThreshold, UnitTestDirectory, ++OperatorCounter, OperatorTotal);
+	FOperatorTester::TestElementWiseOperator<FFloorOperator>(NeuralNetworkInferenceQAAsset, TEXT("FFloorOperator"), ElementWiseTensors, ZeroThreshold, UnitTestDirectory, ++OperatorCounter, OperatorTotal);
+	//UE_LOG(LogNeuralNetworkInferenceQA, Display, TEXT("--------------- %d/%d FGather"), ++OperatorCounter, OperatorTotal);
+	//{
+	//	UE_LOG(LogNeuralNetworkInferenceQA, Display, TEXT("-------------------- Not implemented (yet)"));
+	//}
+	FOperatorTester::TestGemmOperator(NeuralNetworkInferenceQAAsset, InputTensorData, FloatThreshold, UnitTestDirectory, ++OperatorCounter, OperatorTotal);
+	FOperatorTester::TestElementWise1AttributeOperator<FLeakyReluOperator>(NeuralNetworkInferenceQAAsset, TEXT("FLeakyReluOperator"), ElementWiseTensors, ZeroThreshold, UnitTestDirectory, ++OperatorCounter, OperatorTotal);
+	FOperatorTester::TestElementWiseOperator<FLogOperator>(NeuralNetworkInferenceQAAsset, TEXT("FLogOperator"), ElementWiseTensors, FloatThreshold, UnitTestDirectory, ++OperatorCounter, OperatorTotal);
+	FOperatorTester::TestMultiDirectionalBroadcastingOperator<FMulOperator>(NeuralNetworkInferenceQAAsset, TEXT("FMulOperator"), MultidirectionalTensors, ZeroThreshold, UnitTestDirectory, ++OperatorCounter, OperatorTotal);
+	FOperatorTester::TestElementWiseOperator<FNegOperator>(NeuralNetworkInferenceQAAsset, TEXT("FNegOperator"), ElementWiseTensors, ZeroThreshold, UnitTestDirectory, ++OperatorCounter, OperatorTotal);
+	FOperatorTester::TestMultiDirectionalBroadcastingOperator<FPowOperator>(NeuralNetworkInferenceQAAsset, TEXT("FPowOperator"), MultidirectionalTensors, RelaxedFloatThreshold, UnitTestDirectory, ++OperatorCounter, OperatorTotal);
+	FOperatorTester::TestElementWiseOperator<FReciprocalOperator>(NeuralNetworkInferenceQAAsset, TEXT("FReciprocalOperator"), ElementWiseTensors, FloatThreshold, UnitTestDirectory, ++OperatorCounter, OperatorTotal);
+	//UE_LOG(LogNeuralNetworkInferenceQA, Display, TEXT("--------------- %d/%d FReduceSum"), ++OperatorCounter, OperatorTotal);
+	//{
+	//	UE_LOG(LogNeuralNetworkInferenceQA, Display, TEXT("-------------------- Not implemented (yet)"));
+	//}
+	FOperatorTester::TestElementWiseOperator<FReluOperator>(NeuralNetworkInferenceQAAsset, TEXT("FReluOperator"), ElementWiseTensors, ZeroThreshold, UnitTestDirectory, ++OperatorCounter, OperatorTotal);
+	FOperatorTester::TestReshapeOperator(NeuralNetworkInferenceQAAsset, InputTensorData, InputTensorDataPos, ZeroThreshold, UnitTestDirectory, ++OperatorCounter, OperatorTotal);
+	FOperatorTester::TestElementWiseOperator<FRoundOperator>(NeuralNetworkInferenceQAAsset, TEXT("FRoundOperator"), ElementWiseTensors, ZeroThreshold, UnitTestDirectory, ++OperatorCounter, OperatorTotal);
+	//UE_LOG(LogNeuralNetworkInferenceQA, Display, TEXT("--------------- %d/%d FShape"), ++OperatorCounter, OperatorTotal);
+	//{
+	//	UE_LOG(LogNeuralNetworkInferenceQA, Display, TEXT("-------------------- Not implemented (yet)"));
+	//}
+	FOperatorTester::TestElementWiseOperator<FSigmoidOperator>(NeuralNetworkInferenceQAAsset, TEXT("FSigmoidOperator"), ElementWiseTensors, FloatThreshold, UnitTestDirectory, ++OperatorCounter, OperatorTotal);
+	FOperatorTester::TestElementWiseOperator<FSignOperator>(NeuralNetworkInferenceQAAsset, TEXT("FSignOperator"), ElementWiseTensors, ZeroThreshold, UnitTestDirectory, ++OperatorCounter, OperatorTotal);
+	FOperatorTester::TestElementWiseOperator<FSinOperator>(NeuralNetworkInferenceQAAsset, TEXT("FSinOperator"), ElementWiseTensors, VeryRelaxedFloatThreshold, UnitTestDirectory, ++OperatorCounter, OperatorTotal);
+	FOperatorTester::TestElementWiseOperator<FSinhOperator>(NeuralNetworkInferenceQAAsset, TEXT("FSinhOperator"), ElementWiseTensorsSmall, FloatThreshold, UnitTestDirectory, ++OperatorCounter, OperatorTotal);
+	FOperatorTester::TestElementWiseOperator<FSqrtOperator>(NeuralNetworkInferenceQAAsset, TEXT("FSqrtOperator"), ElementWiseTensors, FloatThreshold, UnitTestDirectory, ++OperatorCounter, OperatorTotal);
+	FOperatorTester::TestSqueezeOperator(NeuralNetworkInferenceQAAsset, InputTensorData, InputTensorDataPos, ZeroThreshold, UnitTestDirectory, ++OperatorCounter, OperatorTotal);
+	FOperatorTester::TestMultiDirectionalBroadcastingOperator<FSubOperator>(NeuralNetworkInferenceQAAsset, TEXT("FSubOperator"), MultidirectionalTensors, ZeroThreshold, UnitTestDirectory, ++OperatorCounter, OperatorTotal);
+	FOperatorTester::TestElementWiseOperator<FTanOperator>(NeuralNetworkInferenceQAAsset, TEXT("FTanOperator"), ElementWiseTensors, SuperRelaxedFloatThreshold, UnitTestDirectory, ++OperatorCounter, OperatorTotal);
+	FOperatorTester::TestElementWiseOperator<FTanhOperator>(NeuralNetworkInferenceQAAsset, TEXT("FTanhOperator"), ElementWiseTensorsSmall, FloatThreshold, UnitTestDirectory, ++OperatorCounter, OperatorTotal);
+	//UE_LOG(LogNeuralNetworkInferenceQA, Display, TEXT("--------------- %d/%d FUnsqueeze"), ++OperatorCounter, OperatorTotal);
+	//{
+	//	UE_LOG(LogNeuralNetworkInferenceQA, Display, TEXT("-------------------- Not implemented (yet)"));
+	//}
+	// Verify operator count
+	ensureMsgf(OperatorTotal == OperatorCounter, TEXT("New or missing operators found, OperatorTotal == OperatorCounter failed (%d != %d)"), OperatorTotal, OperatorCounter);
+
+	// Test against GT saved on disk - Only saving a unique UAsset file with all results from all operators
+	if (!NeuralNetworkInferenceQAAsset->CompareNewVsPreviousTests(UnitTestDirectory))
+	{
+		// Saving NeuralNetworkInferenceQAAsset
+		ensureMsgf(NeuralNetworkInferenceQAAsset->Save(), TEXT("NeuralNetworkInferenceQAAsset->Save() failed."));
+		// Error message
+		ensureMsgf(false, TEXT("NeuralNetworkInferenceQAAsset->CompareNewVsPreviousTests() failed. Updated NeuralNetworkInferenceQAAsset saved as %s. Check the previous warning messages."), *NeuralNetworkInferenceQAAsset->GetOutermost()->GetLoadedPath().GetLocalFullPath());
+		// Test failed
+		return false;
+	}
+	// Test successful
+	return true;
+}
+
+
+
+/* FOperatorUnitTester static public functions
+ *****************************************************************************/
+
+void FOperatorUnitTester::CreateRandomArraysOfFixedSeedAndDefinedSizes(TMap<int32, TArray<float>>& InOutInputTensorData, TMap<int32, TArray<float>>& InOutInputTensorDataSmall,
+	TMap<int32, TArray<float>>& InOutInputTensorDataTiny, TMap<int32, TArray<float>>& InOutInputTensorDataPos, const TArray<int32>& InArraySizesToPopulate)
+{
+	FRandomStream InOutRandomStream(99/*FMath::Rand()*/);
+	const float InRange = 1e2;
+	const float InRangeSmall = 4;
+	const float InRangeTiny = 2;
+	for (const int32 ArraySize : InArraySizesToPopulate)
+	{
+		InOutInputTensorData.Add(ArraySize, CreateRandomArray(&InOutRandomStream, ArraySize, -InRange, InRange));
+		InOutInputTensorDataSmall.Add(ArraySize, CreateRandomArray(&InOutRandomStream, ArraySize, -InRangeSmall, InRangeSmall));
+		InOutInputTensorDataTiny.Add(ArraySize, CreateRandomArray(&InOutRandomStream, ArraySize, -InRangeTiny, InRangeTiny));
+		InOutInputTensorDataPos.Add(ArraySize, CreateRandomArray(&InOutRandomStream, ArraySize, 0, InRange));
+	}
+}
+
+TArray<float> FOperatorUnitTester::CreateRandomArray(FRandomStream* InOutRandomStream, const int32 InSize, const float InRangeMin, const float InRangeMax)
+{
+	if (!InOutRandomStream)
+	{
+		ensureMsgf(false, TEXT("InOutRandomStream was a nullptr."));
+		return TArray<float>();
+	}
+	TArray<float> OutArray;
+	OutArray.SetNum(InSize);
+	for (float& FloatValue : OutArray)
+	{
+		FloatValue = InOutRandomStream->FRandRange(InRangeMin, InRangeMax);
+	}
+	return OutArray;
+}
diff --git a/Engine/Plugins/Experimental/NeuralNetworkInference/Source/Private/NeuralNetworkInferenceQA/Private/OperatorUnitTester.h b/Engine/Plugins/Experimental/NeuralNetworkInference/Source/Private/NeuralNetworkInferenceQA/Private/OperatorUnitTester.h
new file mode 100644
--- /dev/null
+++ b/Engine/Plugins/Experimental/NeuralNetworkInference/Source/Private/NeuralNetworkInferenceQA/Private/OperatorUnitTester.h
@@ -0,0 +1,18 @@
+// Copyright Epic Games, Inc. All Rights Reserved.
+
+#pragma once
+
+#include "CoreMinimal.h"
+#include "NeuralTensor.h"
+
+class FOperatorUnitTester
+{
+public:
+	static bool GlobalTest(const FString& InProjectContentDir, const FString& InUnitTestRelativeDirectory);
+
+private:
+	static void CreateRandomArraysOfFixedSeedAndDefinedSizes(TMap<int32, TArray<float>>& InOutInputTensorData, TMap<int32, TArray<float>>& InOutInputTensorDataSmall,
+		TMap<int32, TArray<float>>& InOutInputTensorDataTiny, TMap<int32, TArray<float>>& InOutInputTensorDataPos, const TArray<int32>& InArraySizesToPopulate);
+
+	static TArray<float> CreateRandomArray(struct FRandomStream* InOutRandomStream, const int32 InSize, const float InRangeMin, const float InRangeMax);
+};
diff --git a/Engine/Plugins/Experimental/NeuralNetworkInference/Source/Private/NeuralNetworkInferenceQA/Private/UnitTester.cpp b/Engine/Plugins/Experimental/NeuralNetworkInference/Source/Private/NeuralNetworkInferenceQA/Private/UnitTester.cpp
new file mode 100644
--- /dev/null
+++ b/Engine/Plugins/Experimental/NeuralNetworkInference/Source/Private/NeuralNetworkInferenceQA/Private/UnitTester.cpp
@@ -0,0 +1,42 @@
+// Copyright Epic Games, Inc. All Rights Reserved.
+
+#include "UnitTester.h"
+#include "ModelUnitTester.h"
+#include "NeuralNetworkInferenceQAUtils.h"
+#include "OperatorUnitTester.h"
+
+
+
+/* FUnitTester static public functions
+ *****************************************************************************/
+
+bool FUnitTester::GlobalTest(const FString& InProjectContentDir, const FString& InModelZooRelativeDirectory, const FString& InUnitTestRelativeDirectory)
+{
+	UE_LOG(LogNeuralNetworkInferenceQA, Display, TEXT("---------------------------------------------------------------------------------------------------------------------------------"));
+	UE_LOG(LogNeuralNetworkInferenceQA, Display, TEXT("---------------------------------------------------------------------------------------------------------------------------------"));
+	UE_LOG(LogNeuralNetworkInferenceQA, Display, TEXT("---------------------------------------------------------------------------------------------------------------------------------"));
+	UE_LOG(LogNeuralNetworkInferenceQA, Display, TEXT("----- Starting UnitTesting() ----------------------------------------------------------------------------------------------------"));
+
+	UE_LOG(LogNeuralNetworkInferenceQA, Display, TEXT("---------------------------------------------------------------------------------------------------------------------------------"));
+	UE_LOG(LogNeuralNetworkInferenceQA, Display, TEXT("---------------------------------------------------------------------------------------------------------------------------------"));
+	UE_LOG(LogNeuralNetworkInferenceQA, Display, TEXT("---------------------------------------------------------------------------------------------------------------------------------"));
+	UE_LOG(LogNeuralNetworkInferenceQA, Display, TEXT("---------- 1. Model Unit Testing"));
+	bool bDidGlobalTestPassed = FModelUnitTester::GlobalTest(InProjectContentDir, InModelZooRelativeDirectory);
+
+	UE_LOG(LogNeuralNetworkInferenceQA, Display, TEXT("---------------------------------------------------------------------------------------------------------------------------------"));
+	UE_LOG(LogNeuralNetworkInferenceQA, Display, TEXT("---------------------------------------------------------------------------------------------------------------------------------"));
+	UE_LOG(LogNeuralNetworkInferenceQA, Display, TEXT("---------------------------------------------------------------------------------------------------------------------------------"));
+	UE_LOG(LogNeuralNetworkInferenceQA, Display, TEXT("---------- 2. Operator Unit Testing"));
+	bDidGlobalTestPassed &= FOperatorUnitTester::GlobalTest(InProjectContentDir, InUnitTestRelativeDirectory);
+
+	if (bDidGlobalTestPassed)
+	{
+		UE_LOG(LogNeuralNetworkInferenceQA, Display, TEXT("----- UnitTesting() completed! --------------------------------------------------------------------------------------------------"));
+	}
+	else
+	{
+		UE_LOG(LogNeuralNetworkInferenceQA, Warning, TEXT("----- UnitTesting() finished with warnings/errors! --------------------------------------------------------------------------------------------------"));
+	}
+
+	return bDidGlobalTestPassed;
+}
diff --git a/Engine/Plugins/Experimental/NeuralNetworkInference/Source/Private/NeuralNetworkInferenceQA/Private/UnitTester.h b/Engine/Plugins/Experimental/NeuralNetworkInference/Source/Private/NeuralNetworkInferenceQA/Private/UnitTester.h
new file mode 100644
--- /dev/null
+++ b/Engine/Plugins/Experimental/NeuralNetworkInference/Source/Private/NeuralNetworkInferenceQA/Private/UnitTester.h
@@ -0,0 +1,14 @@
+// Copyright Epic Games, Inc. All Rights Reserved.
+
+#pragma once
+
+#include "CoreMinimal.h"
+
+class FUnitTester
+{
+public:
+	/**
+	 * Main unit tester function.
+	 */
+	static bool GlobalTest(const FString& InProjectContentDir, const FString& InModelZooRelativeDirectory, const FString& InUnitTestRelativeDirectory);
+};
diff --git a/Engine/Plugins/Experimental/NeuralNetworkInference/Source/Private/NeuralNetworkInferenceQA/Public/NeuralNetworkInferenceQA.h b/Engine/Plugins/Experimental/NeuralNetworkInference/Source/Private/NeuralNetworkInferenceQA/Public/NeuralNetworkInferenceQA.h
new file mode 100644
--- /dev/null
+++ b/Engine/Plugins/Experimental/NeuralNetworkInference/Source/Private/NeuralNetworkInferenceQA/Public/NeuralNetworkInferenceQA.h
@@ -0,0 +1,20 @@
+// Copyright Epic Games, Inc. All Rights Reserved.
+
+#pragma once
+
+#include "CoreMinimal.h"
+#include "UObject/Object.h"
+#include "NeuralNetworkInferenceQA.generated.h"
+
+/*
+ * BP wrapper of FUnitTester. It contains the main functions to test this plugin.
+ */
+UCLASS(BlueprintType)
+class NEURALNETWORKINFERENCEQA_API UNeuralNetworkInferenceQA : public UObject
+{
+	GENERATED_BODY()
+
+public:
+	UFUNCTION(BlueprintCallable, Category = "NeuralNetworkInferenceQA")
+	static bool UnitTesting();
+};
