Adding NNI documentation - NNI/Public/NeuraTensor.h partially cleaned and documented:
- TArray<uint8> ArrayCPUForGPUAs32Data removed (wasn't being used).
- ArrayCPU renamed as UnderlyingUInt8ArrayData.
- bool bEnableGPU as transient UPROPERTY.

#rnx

diff --git a/Engine/Plugins/Experimental/NeuralNetworkInference/Source/NeuralNetworkInference/Private/NeuralNetworkImplBackEndUEAndORT.cpp b/Engine/Plugins/Experimental/NeuralNetworkInference/Source/NeuralNetworkInference/Private/NeuralNetworkImplBackEndUEAndORT.cpp
--- a/Engine/Plugins/Experimental/NeuralNetworkInference/Source/NeuralNetworkInference/Private/NeuralNetworkImplBackEndUEAndORT.cpp
+++ b/Engine/Plugins/Experimental/NeuralNetworkInference/Source/NeuralNetworkInference/Private/NeuralNetworkImplBackEndUEAndORT.cpp
@@ -1010,7 +1010,7 @@ void UNeuralNetwork::FImplBackEndUEAndORT::RunSessionImpl(const ENeuralDeviceTyp
 					UploadParameters->Input = InputBufferRef;
 
 					GraphBuilder.AddPass(
-						RDG_EVENT_NAME("NNI:UploadTensor:%s", InputTensor.GetNameData()),
+						RDG_EVENT_NAME("NNI:UploadTensor:%s", *InputTensor.GetName()),
 						FUploadTensorParameters::FTypeInfo::GetStructMetadata(),
 						UploadParameters,
 						ERDGPassFlags::Copy | ERDGPassFlags::NeverCull,
diff --git a/Engine/Plugins/Experimental/NeuralNetworkInference/Source/NeuralNetworkInference/Private/NeuralNetworkInferenceUtils.h b/Engine/Plugins/Experimental/NeuralNetworkInference/Source/NeuralNetworkInference/Private/NeuralNetworkInferenceUtils.h
--- a/Engine/Plugins/Experimental/NeuralNetworkInference/Source/NeuralNetworkInference/Private/NeuralNetworkInferenceUtils.h
+++ b/Engine/Plugins/Experimental/NeuralNetworkInference/Source/NeuralNetworkInference/Private/NeuralNetworkInferenceUtils.h
@@ -34,7 +34,7 @@ public:
 
 	/**
 	 * It blocks the current thread until the RHI thread has finished all instructions before this point.
-	 * @returns The time it waited.
+	 * @return The time it waited.
 	 */
 	static void WaitUntilRHIFinished();
 };
diff --git a/Engine/Plugins/Experimental/NeuralNetworkInference/Source/NeuralNetworkInference/Private/NeuralTensor.cpp b/Engine/Plugins/Experimental/NeuralNetworkInference/Source/NeuralNetworkInference/Private/NeuralTensor.cpp
--- a/Engine/Plugins/Experimental/NeuralNetworkInference/Source/NeuralNetworkInference/Private/NeuralTensor.cpp
+++ b/Engine/Plugins/Experimental/NeuralNetworkInference/Source/NeuralNetworkInference/Private/NeuralTensor.cpp
@@ -131,23 +131,23 @@ ENeuralDataType FPrivateNeuralTensor::GetDataTypeFromTensorProtoDataType(const E
 /* FNeuralTensor costructors
  *****************************************************************************/
 
-FNeuralTensor::FNeuralTensor(const ENeuralDataType InDataType, const TArray<int64>& InSizes, const FString& InName, const ENeuralTensorType InTensorTypeGPU)
-	: Name(InName)
-	, DataType(ENeuralDataType::None)
-	, TensorTypeGPU(InTensorTypeGPU)
+FNeuralTensor::FNeuralTensor(const ENeuralDataType InDataType, const TArray<int64>& InSizes, const FString& InName, const ENeuralTensorType InTensorType)
+	: DataType(ENeuralDataType::None)
+	, Name(InName)
+	, TensorType(InTensorType)
 	, bEnableGPU(false)
 {
 	// Memory allocation
 	SetNumUninitialized(InSizes, InDataType);
 }
 
-FNeuralTensor::FNeuralTensor(const ENeuralDataType InDataType, const int64 InVolume, const FString& InName, const ENeuralTensorType InTensorTypeGPU)
-	: FNeuralTensor(InDataType, InVolume > 0 ? TArray<int64>({ InVolume }) : TArray<int64>({}), InName, InTensorTypeGPU)
+FNeuralTensor::FNeuralTensor(const ENeuralDataType InDataType, const int64 InVolume, const FString& InName, const ENeuralTensorType InTensorType)
+	: FNeuralTensor(InDataType, InVolume > 0 ? TArray<int64>({ InVolume }) : TArray<int64>({}), InName, InTensorType)
 {
 }
 
-FNeuralTensor::FNeuralTensor(const FString& InName, const ENeuralTensorType InTensorTypeGPU)
-	: FNeuralTensor(ENeuralDataType::None, TArray<int64>({}), InName, InTensorTypeGPU)
+FNeuralTensor::FNeuralTensor(const FString& InName, const ENeuralTensorType InTensorType)
+	: FNeuralTensor(ENeuralDataType::None, TArray<int64>({}), InName, InTensorType)
 {
 }
 
@@ -156,8 +156,8 @@ FNeuralTensor::FNeuralTensor(const FString& InName, const ENeuralTensorType InTe
 /* FNeuralTensor private costructor
  *****************************************************************************/
 
-FNeuralTensor::FNeuralTensor(const ENeuralDataType InDataType, const void* const InValues, const int64 InSizeOfT, const int64 InValueNum, const TArray<int64>& InSizes, const FString& InName, const ENeuralTensorType InTensorTypeGPU)
-	: FNeuralTensor(InDataType, InSizes, InName, InTensorTypeGPU)
+FNeuralTensor::FNeuralTensor(const ENeuralDataType InDataType, const void* const InValues, const int64 InSizeOfT, const int64 InValueNum, const TArray<int64>& InSizes, const FString& InName, const ENeuralTensorType InTensorType)
+	: FNeuralTensor(InDataType, InSizes, InName, InTensorType)
 {
 	// Sanity check
 	if (IsEmpty())
@@ -179,8 +179,8 @@ FNeuralTensor::FNeuralTensor(const ENeuralDataType InDataType, const void* const
 FNeuralTensor::FNeuralTensor(const FNeuralTensor& InTensor)
 	: FNeuralTensor(InTensor.DataType, InTensor.Sizes, InTensor.Name)
 {
-	SetTensorTypeGPU(InTensor.TensorTypeGPU);
-	ArrayCPU = InTensor.ArrayCPU;
+	SetTensorTypeGPU(InTensor.TensorType);
+	UnderlyingUInt8ArrayData = InTensor.UnderlyingUInt8ArrayData;
 	bEnableGPU = InTensor.bEnableGPU;
 	PooledBuffer = InTensor.PooledBuffer;
 	BufferSRVRef = InTensor.BufferSRVRef;
@@ -191,8 +191,8 @@ FNeuralTensor& FNeuralTensor::operator=(const FNeuralTensor& InTensor)
 {
 	Name = InTensor.Name;
 	SetNumUninitialized(InTensor.Sizes, InTensor.DataType);
-	SetTensorTypeGPU(InTensor.TensorTypeGPU);
-	ArrayCPU = InTensor.ArrayCPU;
+	SetTensorTypeGPU(InTensor.TensorType);
+	UnderlyingUInt8ArrayData = InTensor.UnderlyingUInt8ArrayData;
 	bEnableGPU = InTensor.bEnableGPU;
 	PooledBuffer = InTensor.PooledBuffer;
 	BufferSRVRef = InTensor.BufferSRVRef;
@@ -205,8 +205,8 @@ FNeuralTensor::FNeuralTensor(FNeuralTensor&& InTensor)
 	: FNeuralTensor(InTensor.DataType, InTensor.Sizes, TEXT(""))
 {
 	Swap(Name, InTensor.Name);
-	SetTensorTypeGPU(InTensor.TensorTypeGPU);
-	Swap(ArrayCPU, InTensor.ArrayCPU);
+	SetTensorTypeGPU(InTensor.TensorType);
+	Swap(UnderlyingUInt8ArrayData, InTensor.UnderlyingUInt8ArrayData);
 	bEnableGPU = InTensor.bEnableGPU;
 	Swap(PooledBuffer, InTensor.PooledBuffer);
 	Swap(BufferSRVRef, InTensor.BufferSRVRef);
@@ -217,8 +217,8 @@ FNeuralTensor& FNeuralTensor::operator=(FNeuralTensor&& InTensor)
 {
 	Swap(Name, InTensor.Name);
 	SetNumUninitialized(InTensor.Sizes, InTensor.DataType);
-	SetTensorTypeGPU(InTensor.TensorTypeGPU);
-	Swap(ArrayCPU, InTensor.ArrayCPU);
+	SetTensorTypeGPU(InTensor.TensorType);
+	Swap(UnderlyingUInt8ArrayData, InTensor.UnderlyingUInt8ArrayData);
 	bEnableGPU = InTensor.bEnableGPU;
 	Swap(PooledBuffer, InTensor.PooledBuffer);
 	Swap(BufferSRVRef, InTensor.BufferSRVRef);
@@ -232,12 +232,12 @@ FNeuralTensor& FNeuralTensor::operator=(FNeuralTensor&& InTensor)
 /* FNeuralTensor public operators
  *****************************************************************************/
 
-bool FNeuralTensor::operator==(const FNeuralTensor& InTensorToCopy) const
+bool FNeuralTensor::operator==(const FNeuralTensor& InTensor) const
 {
 	// Dimensions, sizes, and scalar type match --> Check if data does
-	if (Num() == InTensorToCopy.Num() && GetSizes() == InTensorToCopy.GetSizes())
+	if (DataType == InTensor.DataType && Num() == InTensor.Num() && GetSizes() == InTensor.GetSizes())
 	{
-		return ArrayCPU == InTensorToCopy.ArrayCPU;
+		return UnderlyingUInt8ArrayData == InTensor.UnderlyingUInt8ArrayData;
 	}
 	// Dimensions or sizes or scalar type do not match
 	return false;
@@ -253,7 +253,7 @@ int64 FNeuralTensor::GetSize(const int32 InDimension) const
 	return (InDimension < GetNumberDimensions() ? Sizes[InDimension] : 1u);
 }
 
-void FNeuralTensor::SetTensorTypeGPU(const ENeuralTensorType InTensorTypeGPU)
+void FNeuralTensor::SetTensorTypeGPU(const ENeuralTensorType InTensorType)
 {
 	// Sanity check
 	if (PooledBuffer.IsValid() || BufferSRVRef.IsValid() || BufferUAVRef.IsValid())
@@ -263,8 +263,8 @@ void FNeuralTensor::SetTensorTypeGPU(const ENeuralTensorType InTensorTypeGPU)
 		BufferUAVRef.Reset();
 	}
 
-	// Update TensorTypeGPU
-	TensorTypeGPU = InTensorTypeGPU;
+	// Update TensorType
+	TensorType = InTensorType;
 }
 
 void FNeuralTensor::ToGPU_RenderThread(FRDGBuilder* InOutGraphBuilder)
@@ -282,31 +282,31 @@ void FNeuralTensor::ToGPU_RenderThread(FRDGBuilder* InOutGraphBuilder)
 	//   the input FNeuralTensor). However, Intermediate(Not)Initialized and Output do not copy the memory from CPU to GPU but rather simply allocates it.
 	// - Output might also become Intermediate(Not)Initialized (e.g., if Output -> ReLU -> Output), so it is kept as ReadWrite rather than written once to account for this.
 	// Call ToGPU_RenderThread with the right flags
-	if (TensorTypeGPU == ENeuralTensorType::Generic)
+	if (TensorType == ENeuralTensorType::Generic)
 	{
 		return ToGPU_RenderThread(InOutGraphBuilder, EBufferUsageFlags::ShaderResource | EBufferUsageFlags::UnorderedAccess, true);
 	}
-	else if (TensorTypeGPU == ENeuralTensorType::Input)
+	else if (TensorType == ENeuralTensorType::Input)
 	{
 		return ToGPU_RenderThread(InOutGraphBuilder, EBufferUsageFlags::ShaderResource | EBufferUsageFlags::UnorderedAccess, true);
 	}
-	else if (TensorTypeGPU == ENeuralTensorType::IntermediateNotInitialized)
+	else if (TensorType == ENeuralTensorType::IntermediateNotInitialized)
 	{
 		return ToGPU_RenderThread(InOutGraphBuilder, EBufferUsageFlags::ShaderResource | EBufferUsageFlags::UnorderedAccess | EBufferUsageFlags::Transient, false);
 	}
-	else if (TensorTypeGPU == ENeuralTensorType::IntermediateInitialized)
+	else if (TensorType == ENeuralTensorType::IntermediateInitialized)
 	{
 		return ToGPU_RenderThread(InOutGraphBuilder, EBufferUsageFlags::ShaderResource | EBufferUsageFlags::UnorderedAccess, true);
 	}
-	else if (TensorTypeGPU == ENeuralTensorType::Output)
+	else if (TensorType == ENeuralTensorType::Output)
 	{
 		return ToGPU_RenderThread(InOutGraphBuilder, EBufferUsageFlags::ShaderResource | EBufferUsageFlags::UnorderedAccess, false);
 	}
-	else if (TensorTypeGPU == ENeuralTensorType::Weight)
+	else if (TensorType == ENeuralTensorType::Weight)
 	{
 		return ToGPU_RenderThread(InOutGraphBuilder, EBufferUsageFlags::ShaderResource | EBufferUsageFlags::Static, true);
 	}
-	UE_LOG(LogNeuralNetworkInference, Warning, TEXT("FNeuralTensor-%s::ToGPU_RenderThread(): Unimplemented TensorTypeGPU = %d. Assuming ENeuralTensorType::Generic."), *Name, (int32)TensorTypeGPU);
+	UE_LOG(LogNeuralNetworkInference, Warning, TEXT("FNeuralTensor-%s::ToGPU_RenderThread(): Unimplemented TensorType = %d. Assuming ENeuralTensorType::Generic."), *Name, (int32)TensorType);
 	return ToGPU_RenderThread(InOutGraphBuilder, EBufferUsageFlags::UnorderedAccess, true);
 }
 
@@ -341,7 +341,7 @@ void FNeuralTensor::ToGPU_RenderThread(FRDGBuilder* InOutGraphBuilder, const EBu
 		BufferDesc.UnderlyingType = FRDGBufferDesc::EUnderlyingType::VertexBuffer;
 		
 		BufferRef = bInShouldCopyFromCPU 
-			? CreateVertexBuffer(*InOutGraphBuilder, *Name, BufferDesc, ArrayCPU.GetData(), NumInBytes(), ERDGInitialDataFlags::NoCopy)
+			? CreateVertexBuffer(*InOutGraphBuilder, *Name, BufferDesc, UnderlyingUInt8ArrayData.GetData(), NumInBytes(), ERDGInitialDataFlags::NoCopy)
 			: InOutGraphBuilder->CreateBuffer(BufferDesc, *Name);
 
 		// Recreate PooledBuffer for future runs
@@ -352,7 +352,7 @@ void FNeuralTensor::ToGPU_RenderThread(FRDGBuilder* InOutGraphBuilder, const EBu
 	{
 		BufferRef = InOutGraphBuilder->RegisterExternalBuffer(*PooledBuffer);
 
-		InOutGraphBuilder->QueueBufferUpload(BufferRef, ArrayCPU.GetData(), NumInBytes(), ERDGInitialDataFlags::None);
+		InOutGraphBuilder->QueueBufferUpload(BufferRef, UnderlyingUInt8ArrayData.GetData(), NumInBytes(), ERDGInitialDataFlags::None);
 	}
 
 	// Recreate BufferSRVRef
@@ -419,7 +419,7 @@ void FNeuralTensor::ToCPU_RenderThread(FRDGBuilder* InOutGraphBuilder)
 		const int64 VolumeInBytes = NumInBytes();
 		FRHIBuffer* VertexBuffer = (*PooledBuffer)->GetRHI();
 		const void* const BufferData = RHICmdList.LockBuffer(VertexBuffer, 0, VolumeInBytes, RLM_ReadOnly);
-		FMemory::Memcpy((void*)ArrayCPU.GetData(), BufferData, VolumeInBytes);
+		FMemory::Memcpy((void*)UnderlyingUInt8ArrayData.GetData(), BufferData, VolumeInBytes);
 		RHICmdList.UnlockBuffer(VertexBuffer);
 	});
 }
@@ -523,20 +523,20 @@ void FNeuralTensor::SetNumUninitialized(const TArray<int64>& InSizes, const ENeu
 	}
 	// Update Sizes
 	Sizes = InSizes;
-	// Re-initialize ArrayCPU
+	// Re-initialize UnderlyingUInt8ArrayData
 	if (Sizes.Num() > 0)
 	{
 		Volume = FNeuralNetworkInferenceUtils::Product<int64>(Sizes);
 		const int64 VolumeInBytes = Volume * FNeuralDataTypeUtils::GetByteSize(DataType);
-		if (VolumeInBytes != ArrayCPU.Num())
+		if (VolumeInBytes != UnderlyingUInt8ArrayData.Num())
 		{
-			ArrayCPU.SetNumUninitialized(VolumeInBytes, bInAllowShrinking); // Pre-allocate TArray
+			UnderlyingUInt8ArrayData.SetNumUninitialized(VolumeInBytes, bInAllowShrinking); // Pre-allocate TArray
 		}
 	}
 	else
 	{
 		Volume = 0;
-		ArrayCPU.Empty();
+		UnderlyingUInt8ArrayData.Empty();
 	}
 }
 
@@ -548,10 +548,10 @@ void FNeuralTensor::SetFromUnderlyingUInt8ArrayCopy(const TArray<uint8>& InArray
 			TEXT("FNeuralTensor::SetFromUnderlyingUInt8ArrayCopy(): NumInBytes() == InArray.Num() failed, %d != %d."), NumInBytes(), InArray.Num());
 		return;
 	}
-	ArrayCPU = InArray;
+	UnderlyingUInt8ArrayData = InArray;
 }
 
-bool FNeuralTensor::SetFromTensorProto(const FTensorProto* const InTensorProto, const FString& InTensorName, const ENeuralTensorType InTensorTypeGPU)
+bool FNeuralTensor::SetFromTensorProto(const FTensorProto* const InTensorProto, const FString& InTensorName, const ENeuralTensorType InTensorType)
 {
 	if (!InTensorProto)
 	{
@@ -563,7 +563,7 @@ bool FNeuralTensor::SetFromTensorProto(const FTensorProto* const InTensorProto,
 
 	// Create Tensor
 	Name = InTensorName;
-	TensorTypeGPU = InTensorTypeGPU;
+	TensorType = InTensorType;
 	// Memory allocation
 	SetNumUninitialized(InTensorProto->Dimensions, FPrivateNeuralTensor::GetDataTypeFromTensorProtoDataType(InTensorProto->DataType));
 
@@ -650,11 +650,11 @@ bool FNeuralTensor::Flip(const int32 InDimension)
 		// Remove last index (that makes it a normal index) and replace with its flipped equivalent
 		FlippedTensorIndex = FlippedTensorIndex + Sizes[InDimension] - 1 - 2 * TensorNDIndexes.Last();
 		// Flip TensorIndex value
-		FMemory::Memcpy(&NewArrayOnCPU.GetData()[TensorIndex * BytesPerIndex], &ArrayCPU.GetData()[FlippedTensorIndex * DimensionOffsetInBytes], DimensionOffsetInBytes);
+		FMemory::Memcpy(&NewArrayOnCPU.GetData()[TensorIndex * BytesPerIndex], &UnderlyingUInt8ArrayData.GetData()[FlippedTensorIndex * DimensionOffsetInBytes], DimensionOffsetInBytes);
 		// Increase TensorNDIndexes
 		FPrivateNeuralTensor::NDTensorIndexesPlus1(TensorNDIndexes, TensorNDSizes);
 	}
-	Swap(NewArrayOnCPU, ArrayCPU);
+	Swap(NewArrayOnCPU, UnderlyingUInt8ArrayData);
 	return true;
 }
 
@@ -700,7 +700,7 @@ bool FNeuralTensor::Transpose()
 			TArray<uint8> NewArrayOnCPU;
 			NewArrayOnCPU.SetNumUninitialized(NumInBytes());
 			uint8* NewArrayOnCPUPtr = NewArrayOnCPU.GetData();
-			const uint8* const ArrayOnCPUPtr = ArrayCPU.GetData();
+			const uint8* const ArrayOnCPUPtr = UnderlyingUInt8ArrayData.GetData();
 			const int64 Height = Sizes[0];
 			const int64 Width = Sizes[1];
 			const int64 Bytes = FNeuralDataTypeUtils::GetByteSize(DataType);
@@ -710,11 +710,11 @@ bool FNeuralTensor::Transpose()
 				{
 					for (int64 B = 0; B < Bytes; ++B)
 					{
-						NewArrayOnCPU/*Ptr*/[(X * Height + Y)*Bytes + B] = ArrayCPU/*Ptr*/[(Y * Width + X)*Bytes + B];
+						NewArrayOnCPU/*Ptr*/[(X * Height + Y)*Bytes + B] = UnderlyingUInt8ArrayData/*Ptr*/[(Y * Width + X)*Bytes + B];
 					}
 				}
 			}
-			Swap(NewArrayOnCPU, ArrayCPU);
+			Swap(NewArrayOnCPU, UnderlyingUInt8ArrayData);
 		}
 		else
 		{
@@ -753,12 +753,12 @@ FString FNeuralTensor::ToString(const int64 InMaxNumberElementsToDisplay, const
 		TensorString += (Name.Len() > 0 ? Name : FString(TEXT("Unnamed FNeuralTensor"))) + TEXT(": ")
 			// DataType
 			+ FNeuralDataTypeUtils::ToString(DataType) + TEXT(", ")
-			+ (TensorTypeGPU == ENeuralTensorType::Generic ? TEXT("Generic")
-				: TensorTypeGPU == ENeuralTensorType::Input ? TEXT("Input")
-				: TensorTypeGPU == ENeuralTensorType::IntermediateNotInitialized ? TEXT("IntermediateNotInitialized")
-				: TensorTypeGPU == ENeuralTensorType::IntermediateInitialized ? TEXT("IntermediateInitialized")
-				: TensorTypeGPU == ENeuralTensorType::Output ? TEXT("Output")
-				: TensorTypeGPU == ENeuralTensorType::Weight ? TEXT("Weight")
+			+ (TensorType == ENeuralTensorType::Generic ? TEXT("Generic")
+				: TensorType == ENeuralTensorType::Input ? TEXT("Input")
+				: TensorType == ENeuralTensorType::IntermediateNotInitialized ? TEXT("IntermediateNotInitialized")
+				: TensorType == ENeuralTensorType::IntermediateInitialized ? TEXT("IntermediateInitialized")
+				: TensorType == ENeuralTensorType::Output ? TEXT("Output")
+				: TensorType == ENeuralTensorType::Weight ? TEXT("Weight")
 				: TEXT("Unknown"))
 			// Volume and sizes
 			+ TEXT(", volume=") + FString::FromInt(Num()) + TEXT(", sizes={");
@@ -874,7 +874,7 @@ void FNeuralTensor::SetFromPointer(const void* const InData, const int64 InSizeO
 	else
 	{
 		// Deep copy
-		FMemory::Memcpy(ArrayCPU.GetData(), InData, NumInBytes());
+		FMemory::Memcpy(UnderlyingUInt8ArrayData.GetData(), InData, NumInBytes());
 	}
 }
 
diff --git a/Engine/Plugins/Experimental/NeuralNetworkInference/Source/NeuralNetworkInference/Private/UEOnly/GraphProtoToNeuralNetworkConverter.cpp b/Engine/Plugins/Experimental/NeuralNetworkInference/Source/NeuralNetworkInference/Private/UEOnly/GraphProtoToNeuralNetworkConverter.cpp
--- a/Engine/Plugins/Experimental/NeuralNetworkInference/Source/NeuralNetworkInference/Private/UEOnly/GraphProtoToNeuralNetworkConverter.cpp
+++ b/Engine/Plugins/Experimental/NeuralNetworkInference/Source/NeuralNetworkInference/Private/UEOnly/GraphProtoToNeuralNetworkConverter.cpp
@@ -323,15 +323,15 @@ bool FGraphProtoToNeuralNetworkConverter::CreateOperatorsAndEditTensorArray(TArr
 					{
 						// Is it also an absolute output of the network?
 						const int64 TensorIndex = InOutTensors.Num();
-						ENeuralTensorType TensorTypeGPU = ENeuralTensorType::IntermediateNotInitialized;
+						ENeuralTensorType TensorType = ENeuralTensorType::IntermediateNotInitialized;
 						if (InOutputNameDummyIndexMap.Find(OutputTensorName))
 						{
 							OutputNameIndexMap.Add(OutputTensorName, TensorIndex);
-							TensorTypeGPU = ENeuralTensorType::Output;
+							TensorType = ENeuralTensorType::Output;
 						}
 						// Create tensor
 						InOutNameIndexMap.Add(OutputTensorName, TensorIndex);
-						InOutTensors.Push(FNeuralTensor(ENeuralDataType::Float, TArray<int64>({}), OutputTensorName, TensorTypeGPU));
+						InOutTensors.Push(FNeuralTensor(ENeuralDataType::Float, TArray<int64>({}), OutputTensorName, TensorType));
 						OperatorOutputTensors.Push(&InOutTensors.Last());
 					}
 					// Subcase 3/3: Output tensor is the same than 1 of the input tensors of that operator (inlined operator)
diff --git a/Engine/Plugins/Experimental/NeuralNetworkInference/Source/NeuralNetworkInference/Public/NeuralTensor.h b/Engine/Plugins/Experimental/NeuralNetworkInference/Source/NeuralNetworkInference/Public/NeuralTensor.h
--- a/Engine/Plugins/Experimental/NeuralNetworkInference/Source/NeuralNetworkInference/Public/NeuralTensor.h
+++ b/Engine/Plugins/Experimental/NeuralNetworkInference/Source/NeuralNetworkInference/Public/NeuralTensor.h
@@ -28,98 +28,107 @@ struct NEURALNETWORKINFERENCE_API FNeuralTensor
 
 public:
 	/**
-	 * It allocates the desired memory.
-	 * @param InVolume Set to 0 if memory allocation is not required or the final size is unknown. Values smaller than 0 will be clipped to 0.
-	 * @param InSizes Set to empty (or omit argument) if memory allocation is not required or the final size is unknown.
-	 * @param InName Used for GPU debugging and the ToString() function.
-	*/
-	explicit FNeuralTensor(const ENeuralDataType InDataType = ENeuralDataType::None, const TArray<int64>& InSizes = TArray<int64>(),
-		const FString& InName = TEXT("FNeuralTensor"), const ENeuralTensorType InTensorTypeGPU = ENeuralTensorType::Generic);
-	explicit FNeuralTensor(const ENeuralDataType InDataType, const int64 InVolume, const FString& InName = TEXT("FNeuralTensor"),
-		const ENeuralTensorType InTensorTypeGPU = ENeuralTensorType::Generic);
-	explicit FNeuralTensor(const FString& InName, const ENeuralTensorType InTensorTypeGPU = ENeuralTensorType::Generic);
-
-	/**
-	 * Performance-wise, this constructor makes a deep copy of the data (not optimal). For maximum speed, use the other constructors + GetData()/GetDataCasted().
-	 * It allocates the desired memory, and fills it with the input data from InArray. Equivalent to FNeuralTensor(InVolumeOrSizes) + SetFromArrayCopy(...)
+	 * There are several constructors of FNeuralTensor, all of them similar to each other. They set and fill the main variables of this class:
 	 *
-	 * @param InArray The input data to copy from.
-	 * @param InSizes If empty (default), a 1-D tensor of Volume = InArray.Num() will be assumed. If not empty, used to fill the Sizes of this FNeuralTensor.
-	 * @param InName Used for GPU debugging and the ToString() function.
-	 */
+	 * @param InDataType It sets DataType, which defines the underlying uint8 data type of the network (float, double, int32, etc).
+	 * @param InSizes It sets Sizes and Volume. Sizes defines the dimensions of the tensor and Volume the total number of elements (mathematically
+	 *  Prod(Sizes[Index]) for all Indexes). Set to empty (or omit this argument) if memory allocation is not required or the final tensor size is
+	 *  unknown.
+	 * @param InVolume Alternative to InSizes which also sets Sizes and Volume. Using InVolume is equivalent to using
+	 *  InSizes = TArray<int64>({InVolume}). I.e., it behaves like InSizes, but assumes Sizes.Num() == 1 and sets Sizes[0] = Volume = InVolume. Set
+	 *  to 0 or a negative value if memory allocation is not required or the final size is unknown.
+	 * @param InName It sets Name, which is used for GPU debugging and the ToString() function.
+	 * @param InTensorType It sets TensorType, which is used when moving memory to the GPU.
+	 * @param InArray Alternative to InSizes/InVolume, and represents the input data to copy from. The templated constructor that uses InArray is a
+	 *  simple (but not efficient) way of initializing a FNeuralTensor from the existing InArray TArray. It is equivalent to calling
+	 *  FNeuralTensor(..., InSizes/InVolume, ...) and then SetFromArrayCopy(InArray). It is not efficient because it makes a deep copy of the data
+	 *  of InArray. For maximum speed, rather than creating this intermediate InArray TArray, you can avoid the double copy by calling
+	 *  FNeuralTensor(..., InSizes/InVolume, ...) and then filling the tensor memory with GetData()/GetDataCasted().
+	 *
+	 * In addition, the array containing the CPU data, UnderlyingUInt8ArrayData, is pre-allocated based on the combination of InDataType and
+	 * InSizes/InVolume (if InSizes/InVolume not empty/zero): UnderlyingUInt8ArrayData.Num() == sizeof(data type) x Volume.
+	 */
+	FNeuralTensor(const ENeuralDataType InDataType = ENeuralDataType::None, const TArray<int64>& InSizes = TArray<int64>(),
+		const FString& InName = TEXT("FNeuralTensor"), const ENeuralTensorType InTensorType = ENeuralTensorType::Generic);
+	FNeuralTensor(const ENeuralDataType InDataType, const int64 InVolume, const FString& InName = TEXT("FNeuralTensor"),
+		const ENeuralTensorType InTensorType = ENeuralTensorType::Generic);
+	FNeuralTensor(const FString& InName, const ENeuralTensorType InTensorType = ENeuralTensorType::Generic);
 	template <typename T>
-	explicit FNeuralTensor(const TArray<T>& InArray, const TArray<int64>& InSizes = TArray<int64>(), const FString& InName = TEXT("FNeuralTensor"),
-		const ENeuralTensorType InTensorTypeGPU = ENeuralTensorType::Generic);
+	FNeuralTensor(const TArray<T>& InArray, const TArray<int64>& InSizes = TArray<int64>(), const FString& InName = TEXT("FNeuralTensor"),
+		const ENeuralTensorType InTensorType = ENeuralTensorType::Generic);
 
 	/**
-	 * Copy constructor.
-	 * It performs a deep (slow but safe) copy of the current tensor. The resulting tensor will not share any parameters with the current one.
+	 * Copy and move constructors and assignment functions, respectively.
+	 * - Copy constructor and assignment (input is a const FNeuralTensor): They perform a deep (slow but safe) copy of the current tensor. The
+	 *   resulting tensor will not share any parameters with the current one and both can be safely used.
+	 * - Move constructor and assignment (input is not const and will not be valid after calling this function): They move the members of the input
+	 *   InTensor FNeuralTensor, thus InTensor is no longer safe to use.
+	 *
+	 * @param InTensor The input FNeuralTensor to copy/move.
 	 */
 	FNeuralTensor(const FNeuralTensor& InTensor);
-
-	/**
-	 * Copy assignment. Analog to FNeuralTensor(const FNeuralTensor& InTensor).
-	 */
 	FNeuralTensor& operator=(const FNeuralTensor& InTensor);
-
-	/**
-	 * Move constructor. It destroys the original FNeuralTensor (InTensor) to be moved for a fast Move construction.
-	 */
 	FNeuralTensor(FNeuralTensor&& InTensor);
-
-	/**
-	 * Move assignment. Analog to FNeuralTensor(FNeuralTensor&& InTensor).
-	 */
 	FNeuralTensor& operator=(FNeuralTensor&& InTensor);
 
 	/**
-	 * Comparison operator (equal). Returns true if the dimensions, sizes, scalar type, and data match with each other.
-	 * It does not consider other properties of the FNeuralTensor (such as ENeuralTensorType).
-	 */
-	bool operator==(const FNeuralTensor& InTensorToCopy) const;
-
-	/**
-	 * Comparison operator (not equal). Returns the opposite than the equal (==) operator.
-	 * It does not consider other properties of the FNeuralTensor (such as ENeuralTensorType).
+	 * Comparison operators (equal and not equal). They compare whether the CPU variables match between the 2 FNeuralTensor's (i.e., DataType, Sizes,
+	 * Volume and CPU data). They do not consider other FNeuralTensor variables (such as ENeuralTensorType).
+	 * @return
+	 * - Equal operator, operator==(): True if the CPU variables matched each other, false otherwise.
+	 * - Not equal operator, operator!=() (opposite than the equal operator): False if the CPU variables matched each other, false otherwise.
 	 */
-	FORCEINLINE bool operator!=(const FNeuralTensor& InTensorToCopy) const;
+	bool operator==(const FNeuralTensor& InTensor) const;
+	FORCEINLINE bool operator!=(const FNeuralTensor& InTensor) const;
 
 	/**
-	 * Returns the reference to an element at the given index.
-	 * @returns Reference to the indexed element.
+	 * @see At(), GetArrayCopy(), GetUnderlyingUInt8ArrayRef(), GetData(), GetDataCasted(), which represent all the functions that can access/modify
+	 *  the CPU memory. These functions could result in undefined behavior in not used properly:
+	 * - If FNeuralTensor goes out of scope, it is destructed, or its memory is reset (constructor, SetNum(), etc.), these reference/pointer returned
+	 *   by any of these functions will no longer be valid (except GetUnderlyingUInt8ArrayRef()).
+	 * - These functions only modify the CPU memory, they do not synchronize CPU and GPU memory automatically. The user must do this accordingly.
+	 *
+	 * At() returns a reference to the desired indexed element of the underlying CPU data (UnderlyingUInt8ArrayData). There are 2 variations of this
+	 * function, one constant (it does not allow to modify the referenced value) and a non-const one (which allows modifying it).
+	 *
+	 * @return Reference to the indexed element (const or non-constant).
 	 */
 	template <typename T, typename TInput>
 	T& At(const TInput InIndex);
-
-	/**
-	 * Const version of the non-const function At().
-	 * @returns Reference to the indexed element.
-	 */
 	template <typename T, typename TInput>
 	const T& At(const TInput InIndex) const;
 
 	/**
-	 * 2 different functions to get the output tensor results:
-	 * - GetArrayCopy<T> (slower but safer): It returns a copy of the results as a TArray<T>. T has to be the same size than sizeof(DataType).
-	 * - GetUnderlyingUInt8ArrayRef (faster): It returns a reference to the TArray<uint8> that contains the results. The returned reference will
-	 *   only be valid as long as this class instance does not go out of scope and its internal TArray<uint8> is not re-initialized.
-	 * @see GetArrayCopy(), GetUnderlyingUInt8ArrayRef()
+	 * @see At(), GetArrayCopy(), GetUnderlyingUInt8ArrayRef(), GetData(), GetDataCasted(), which represent all the functions that can access/modify
+	 *  the CPU memory. These functions could result in undefined behavior in not used properly:
+	 * - If FNeuralTensor goes out of scope, it is destructed, or its memory is reset (constructor, SetNum(), etc.), these reference/pointer returned
+	 *   by any of these functions will no longer be valid (except GetUnderlyingUInt8ArrayRef()).
+	 * - These functions only modify the CPU memory, they do not synchronize CPU and GPU memory automatically. The user must do this accordingly.
+	 *
+	 * There are 2 functions to access the tensor as a TArray (but none of them can modify the tensor underlying memory):
+	 * - GetArrayCopy<T> (slower but safer) returns a copy of the data as a TArray<T>. T has to be the same size than sizeof(DataType).
+	 * - GetUnderlyingUInt8ArrayRef (faster but could go out of scope) returns a reference to the underlying TArray<uint8> that contains the results.
 	 */
 	template <typename T>
 	TArray<T> GetArrayCopy() const;
 	FORCEINLINE const TArray<uint8>& GetUnderlyingUInt8ArrayRef() const;
 
 	/**
-	 * Dangerous: If FNeuralTensor goes out of scope, this returned pointer will no longer be valid.
-	 * Analog to TArray::GetData().
-	 * @returns Void pointer to the FNeuralTensor data.
+	 * @see At(), GetArrayCopy(), GetUnderlyingUInt8ArrayRef(), GetData(), GetDataCasted(), which represent all the functions that can access/modify
+	 *  the CPU memory. These functions could result in undefined behavior in not used properly:
+	 * - If FNeuralTensor goes out of scope, it is destructed, or its memory is reset (constructor, SetNum(), etc.), these reference/pointer returned
+	 *   by any of these functions will no longer be valid (except GetUnderlyingUInt8ArrayRef()).
+	 * - These functions only modify the CPU memory, they do not synchronize CPU and GPU memory automatically. The user must do this accordingly.
+	 *
+	 * Both GetData() and GetDataCasted() return a pointer of the underlying CPU data (UnderlyingUInt8ArrayData):
+	 * - GetData() returns a void pointer.
+	 * - GetDataCasted() returns a pointer casted to the desired type T.
+	 * There are 2 variations of each of those 2 functions, one constant (it does not allow to modify the pointer values) and a non-const one (which
+	 * allows modifying it).
+	 *
+	 * @return Void (GetData) or T (GetDataCasted) pointer to the FNeuralTensor data (const or non-constant).
 	 */
 	FORCEINLINE void* GetData();
-
-	/**
-	 * Const analog to TArray::GetData().
-	 * @returns Const void pointer to the FNeuralTensor data.
-	 */
 	FORCEINLINE const void* const GetData() const;
 
 	template<typename T>
@@ -128,8 +137,14 @@ public:
 	template<typename T>
 	FORCEINLINE const T* const GetDataCasted() const;
 
-	FORCEINLINE FString GetName() const;
+	/**
+	 * It returns a copy of the name string.
+	 */
+	FORCEINLINE const FString& GetName() const;
 
+	/**
+	 * It returns a TCHAR pointer of the name string.
+	 */
 	FORCEINLINE const TCHAR* GetNameData() const;
 
 	FORCEINLINE ENeuralDataType GetDataType() const;
@@ -167,13 +182,13 @@ public:
 	 * It sets the ENeuralTensorType (Generic, Input, Intermediate(Not)Initialized, Output, Weight, etc.).
 	 * If the GPU memory was already initialized, it will also UE_LOG a warning.
 	 */
-	void SetTensorTypeGPU(const ENeuralTensorType InTensorTypeGPU);
+	void SetTensorTypeGPU(const ENeuralTensorType InTensorType);
 
 	FORCEINLINE void SetEnableGPU(const bool bInEnableGPU);
 
 	/**
-	 * It uploads/downloads the memory from/to the CPU to/from the GPU based on TensorTypeGPU (which sets a preset of EBufferUsageFlags flags).
-	 * @param InEBufferUsageFlags gives the user total control over the buffer flags (and ignores the TensorTypeGPU flag). This is meant to be filled
+	 * It uploads/downloads the memory from/to the CPU to/from the GPU based on TensorType (which sets a preset of EBufferUsageFlags flags).
+	 * @param InEBufferUsageFlags gives the user total control over the buffer flags (and ignores the TensorType flag). This is meant to be filled
 	 * with a combination of EBufferUsageFlags values.
 	 */
 	void ToGPU_RenderThread(FRDGBuilder* InOutGraphBuilder);
@@ -218,9 +233,9 @@ public:
 
 	/**
 	 * It will fill the current neural tensor with the input FTensorProto.
-	 * @returns Whether the conversion was successful.
+	 * @return Whether the conversion was successful.
 	 */
-	bool SetFromTensorProto(const struct FTensorProto* const InTensorProto, const FString& InTensorName, const ENeuralTensorType InTensorTypeGPU);
+	bool SetFromTensorProto(const struct FTensorProto* const InTensorProto, const FString& InTensorName, const ENeuralTensorType InTensorType);
 
 	/**
 	 * It sets all the elements of the FNeuralTensor to InValue.
@@ -263,41 +278,42 @@ public:
 	FString ToString(const int64 InMaxNumberElementsToDisplay = -1, const bool bInReturnOnlyData = false) const;
 
 protected:
-	/** Used for GPU debugging and the ToString() function. */
-	UPROPERTY(VisibleAnywhere, Category = "Neural Network Inference")
-	FString Name;
-	/** General variables and properties */
+	/**
+	 * General variables and properties set by the FNeuralTensor constructor and saved on the FNeuralTensor UAsset.
+	 * @see FNeuralTensor::FNeuralTensor() for more details about each one of these members.
+	 */
 	UPROPERTY(VisibleAnywhere, Category = "Neural Network Inference")
 	ENeuralDataType DataType;
 	UPROPERTY(VisibleAnywhere, Category = "Neural Network Inference")
+	TArray<int64> Sizes;
+	UPROPERTY(VisibleAnywhere, Category = "Neural Network Inference")
 	int64 Volume;
 	UPROPERTY(VisibleAnywhere, Category = "Neural Network Inference")
-	TArray<int64> Sizes;
+	FString Name;
 	UPROPERTY(VisibleAnywhere, Category = "Neural Network Inference")
-	ENeuralTensorType TensorTypeGPU;
+	ENeuralTensorType TensorType;
 
 private:
 	/** CPU-based variables */
 	UPROPERTY()
-	TArray<uint8> ArrayCPU;
+	TArray<uint8> UnderlyingUInt8ArrayData;
 	/**
 	 * GPU-based variables that are transient (do not need to be serialized) and initialized on the fly (with ToGPU/ToCPU).
 	 * - bEnableGPU: By default false, meaning all GPU memory will be disabled and those functions will not do anything. Enable to allow using the
 	 *   GPU functions and variables of the tensor.
-	 * - ArrayCPUForGPUAs32Data: If ArrayCPU is meant for 64-byte data (i.e., int64, uint64, double).
 	 */
+	UPROPERTY(Transient, VisibleAnywhere, Category = "Neural Network Inference")
 	bool bEnableGPU;
 	TSharedPtr<TRefCountPtr<FRDGPooledBuffer>> PooledBuffer;
 	TSharedPtr<FRDGBufferSRVRef> BufferSRVRef;
 	TSharedPtr<FRDGBufferUAVRef> BufferUAVRef;
-	TArray<uint8> ArrayCPUForGPUAs32Data;
 
 	/**
 	 * Auxiliary for the templated constructor FNeuralTensor(const ENeuralDataType InDataType, const TArray<T>& InArray, ...).
 	 * InSizeOfT and sizeof(InDataType) should match, as well as the volume of InValueNum and InSizes.
 	 */
-	explicit FNeuralTensor(const ENeuralDataType InDataType, const void* const InValues, const int64 InSizeOfT, const int64 InValueNum,
-		const TArray<int64>& InSizes, const FString& InName, const ENeuralTensorType InTensorTypeGPU);
+	FNeuralTensor(const ENeuralDataType InDataType, const void* const InValues, const int64 InSizeOfT, const int64 InValueNum,
+		const TArray<int64>& InSizes, const FString& InName, const ENeuralTensorType InTensorType);
 
 	/**
 	 * Auxiliary for the templated SetFromArrayCopy().
@@ -318,28 +334,28 @@ private:
  *****************************************************************************/
 
 template <typename T>
-FNeuralTensor::FNeuralTensor(const TArray<T>& InArray, const TArray<int64>& InSizes, const FString& InName, const ENeuralTensorType InTensorTypeGPU)
+FNeuralTensor::FNeuralTensor(const TArray<T>& InArray, const TArray<int64>& InSizes, const FString& InName, const ENeuralTensorType InTensorType)
 	: FNeuralTensor(FNeuralDataTypeUtils::GetDataType<T>(), InArray.GetData(), sizeof(T), InArray.Num(),
-		(InSizes.Num() > 0 ? InSizes : TArray<int64>({ (int64)InArray.Num() })), InName, InTensorTypeGPU)
+		(InSizes.Num() > 0 ? InSizes : TArray<int64>({ (int64)InArray.Num() })), InName, InTensorType)
 {}
 
-bool FNeuralTensor::operator!=(const FNeuralTensor& InTensorToCopy) const
+bool FNeuralTensor::operator!=(const FNeuralTensor& InTensor) const
 {
-	return !(*this == InTensorToCopy);
+	return !(*this == InTensor);
 }
 
 template <typename T, typename TInput>
 T& FNeuralTensor::At(const TInput InIndex)
 {
 	checkf(CheckTAndDataTypeEquivalent<T>(), TEXT("FNeuralTensor::At(): CheckTAndDataType failed."));
-	return ((T*)ArrayCPU.GetData())[InIndex];
+	return ((T*)UnderlyingUInt8ArrayData.GetData())[InIndex];
 }
 
 template <typename T, typename TInput>
 const T& FNeuralTensor::At(const TInput InIndex) const
 {
 	checkf(CheckTAndDataTypeEquivalent<T>(), TEXT("FNeuralTensor::At(): CheckTAndDataType failed."));
-	return ((T*)ArrayCPU.GetData())[InIndex];
+	return ((T*)UnderlyingUInt8ArrayData.GetData())[InIndex];
 }
 
 template <typename T>
@@ -349,41 +365,41 @@ TArray<T> FNeuralTensor::GetArrayCopy() const
 	if (CheckTAndDataTypeEquivalent<T>())
 	{
 		Array.SetNumUninitialized(Num());
-		FMemory::Memcpy(Array.GetData(), ArrayCPU.GetData(), NumInBytes());
+		FMemory::Memcpy(Array.GetData(), UnderlyingUInt8ArrayData.GetData(), NumInBytes());
 	}
 	return Array;
 }
 
 const TArray<uint8>& FNeuralTensor::GetUnderlyingUInt8ArrayRef() const
 {
-	return ArrayCPU;
+	return UnderlyingUInt8ArrayData;
 }
 
 void* FNeuralTensor::GetData()
 {
-	return (void*)ArrayCPU.GetData();
+	return (void*)UnderlyingUInt8ArrayData.GetData();
 }
 
 const void* const FNeuralTensor::GetData() const
 {
-	return (void*)ArrayCPU.GetData();
+	return (void*)UnderlyingUInt8ArrayData.GetData();
 }
 
 template<typename T>
 T* FNeuralTensor::GetDataCasted()
 {
 	checkf(CheckTAndDataTypeEquivalent<T>(), TEXT("FNeuralTensor::GetDataCasted(): CheckTAndDataType failed."));
-	return (T*)ArrayCPU.GetData();
+	return (T*)UnderlyingUInt8ArrayData.GetData();
 }
 
 template<typename T>
 const T* const FNeuralTensor::GetDataCasted() const
 {
 	checkf(CheckTAndDataTypeEquivalent<T>(), TEXT("FNeuralTensor::GetDataCasted(): CheckTAndDataType failed."));
-	return (T*)ArrayCPU.GetData();
+	return (T*)UnderlyingUInt8ArrayData.GetData();
 }
 
-FString FNeuralTensor::GetName() const
+const FString& FNeuralTensor::GetName() const
 {
 	return Name;
 }
@@ -410,7 +426,7 @@ int32 FNeuralTensor::GetNumberDimensions() const
 
 int64 FNeuralTensor::NumInBytes() const
 {
-	return ArrayCPU.Num();
+	return UnderlyingUInt8ArrayData.Num();
 }
 
 int64 FNeuralTensor::Num() const
@@ -420,12 +436,12 @@ int64 FNeuralTensor::Num() const
 
 bool FNeuralTensor::IsEmpty() const
 {
-	return ArrayCPU.IsEmpty();
+	return UnderlyingUInt8ArrayData.IsEmpty();
 }
 
 ENeuralTensorType FNeuralTensor::GetTensorTypeGPU() const
 {
-	return TensorTypeGPU;
+	return TensorType;
 }
 
 void FNeuralTensor::SetEnableGPU(const bool bInEnableGPU)
