NNI/Public/NeuraTensor.h: Removed redundant GetNameData (using GetName instead).

#rnx

diff --git a/Engine/Plugins/Experimental/NeuralNetworkInference/Source/NeuralNetworkInference/Private/NeuralTensor.cpp b/Engine/Plugins/Experimental/NeuralNetworkInference/Source/NeuralNetworkInference/Private/NeuralTensor.cpp
--- a/Engine/Plugins/Experimental/NeuralNetworkInference/Source/NeuralNetworkInference/Private/NeuralTensor.cpp
+++ b/Engine/Plugins/Experimental/NeuralNetworkInference/Source/NeuralNetworkInference/Private/NeuralTensor.cpp
@@ -446,7 +446,7 @@ bool FNeuralTensor::InitPooledBuffer(void** NativeResource)
 			BufferDesc.Usage = EBufferUsageFlags::UnorderedAccess | EBufferUsageFlags::ShaderResource;
 			BufferDesc.UnderlyingType = FRDGBufferDesc::EUnderlyingType::VertexBuffer;
 
-			FRDGBufferRef BufferRef = Builder.CreateBuffer(BufferDesc, GetNameData());
+			FRDGBufferRef BufferRef = Builder.CreateBuffer(BufferDesc, *GetName());
 
 			// Recreate PooledBuffer for future runs
 			PooledBuffer = MakeShared<TRefCountPtr<FRDGPooledBuffer>>();
diff --git a/Engine/Plugins/Experimental/NeuralNetworkInference/Source/NeuralNetworkInference/Public/NeuralTensor.h b/Engine/Plugins/Experimental/NeuralNetworkInference/Source/NeuralNetworkInference/Public/NeuralTensor.h
--- a/Engine/Plugins/Experimental/NeuralNetworkInference/Source/NeuralNetworkInference/Public/NeuralTensor.h
+++ b/Engine/Plugins/Experimental/NeuralNetworkInference/Source/NeuralNetworkInference/Public/NeuralTensor.h
@@ -15,11 +15,14 @@
 /**
  * For a general overview of NeuralNetworkInference (NNI), including documentation and code samples, @see UNeuralNetwork, the main class of NNI.
  *
+ * NNI's users should not directly create or modify instances of FNeuralTensor. Instead, they should only interact with FNeuralNetwork and its input
+ * and output functions. E.g., @see FNeuralNetwork::GetInputTensor, FNeuralNetwork::SetInputFromArrayCopy, FNeuralNetwork::GetOutputTensor, etc.
+ *
  * FNeuralTensor is an auxiliary class of UNeuralNetwork which represents a tensor of the UNeuralNetwork model. It is Unreal Engine's equivalent of
- * torch.Tensor (PyTorch) or caffe::Blob.
+ * torch.Tensor (PyTorch) or caffe::Blob (Caffe).
  *
- * Most of its functions run on the CPU, so `ToGPU_RenderThread()` must be called before running on GPU and after running any FNeuralTensor function
- * that modifies the CPU memory. In addition, FNeuralTensor's CPU functions are very similar to those of TArray<T>.
+ * Most of FNeuralTensor's functions run on the CPU, so `ToGPU_RenderThread()` must be called before running on GPU and after running any
+ * FNeuralTensor function that modifies the CPU memory. In addition, FNeuralTensor's CPU functions are very similar to those of TArray<T>.
  */
 USTRUCT()
 struct NEURALNETWORKINFERENCE_API FNeuralTensor
@@ -142,11 +145,6 @@ public:
 	 */
 	FORCEINLINE const FString& GetName() const;
 
-	/**
-	 * It returns a TCHAR pointer of the name string.
-	 */
-	FORCEINLINE const TCHAR* GetNameData() const;
-
 	FORCEINLINE ENeuralDataType GetDataType() const;
 
 	/**
@@ -404,11 +402,6 @@ const FString& FNeuralTensor::GetName() const
 	return Name;
 }
 
-const TCHAR* FNeuralTensor::GetNameData() const
-{
-	return Name.GetCharArray().GetData();
-}
-
 ENeuralDataType FNeuralTensor::GetDataType() const
 {
 	return DataType;
