Adding NNI documentation:
- NNI/Public/NeuraTensor.h more documented and re-ordered a bit.

#rnx

diff --git a/Engine/Plugins/Experimental/NeuralNetworkInference/Source/NeuralNetworkInference/Private/NeuralNetworkImplBackEndUEAndORT.cpp b/Engine/Plugins/Experimental/NeuralNetworkInference/Source/NeuralNetworkInference/Private/NeuralNetworkImplBackEndUEAndORT.cpp
--- a/Engine/Plugins/Experimental/NeuralNetworkInference/Source/NeuralNetworkInference/Private/NeuralNetworkImplBackEndUEAndORT.cpp
+++ b/Engine/Plugins/Experimental/NeuralNetworkInference/Source/NeuralNetworkInference/Private/NeuralNetworkImplBackEndUEAndORT.cpp
@@ -785,7 +785,7 @@ bool UNeuralNetwork::FImplBackEndUEAndORT::SetTensorsFromNetwork(TArray<FNeuralT
 	{
 		for (int32 TensorIndex = 0; TensorIndex < TensorNumber; ++TensorIndex)
 		{
-			OutTensors[TensorIndex].SetTensorTypeGPU(InTensorGPUTypes[TensorIndex]);
+			OutTensors[TensorIndex].SetTensorType(InTensorGPUTypes[TensorIndex]);
 		}
 	}
 
@@ -992,7 +992,7 @@ void UNeuralNetwork::FImplBackEndUEAndORT::RunSessionImpl(const ENeuralDeviceTyp
 
 		for (auto& InputTensor : InputTensors)
 		{
-			if (InputTensor.GetTensorTypeGPU() != ENeuralTensorType::Input)
+			if (InputTensor.GetTensorType() != ENeuralTensorType::Input)
 				continue;
 
 			bNeedsGPUCopy = true;
diff --git a/Engine/Plugins/Experimental/NeuralNetworkInference/Source/NeuralNetworkInference/Private/NeuralOperators/ConvTransposeOperator.cpp b/Engine/Plugins/Experimental/NeuralNetworkInference/Source/NeuralNetworkInference/Private/NeuralOperators/ConvTransposeOperator.cpp
--- a/Engine/Plugins/Experimental/NeuralNetworkInference/Source/NeuralNetworkInference/Private/NeuralOperators/ConvTransposeOperator.cpp
+++ b/Engine/Plugins/Experimental/NeuralNetworkInference/Source/NeuralNetworkInference/Private/NeuralOperators/ConvTransposeOperator.cpp
@@ -456,7 +456,7 @@ bool FConvTransposeOperator::SetAndConfigureAuxiliaryTensor()
 		// Initialize XWithZeros
 		XWithZeros.SetNumUninitialized(XWithZerosSizes, X.GetDataType());
 		XWithZeros.SetTo<float>(0.f);
-		XWithZeros.SetTensorTypeGPU(ENeuralTensorType::IntermediateInitialized);
+		XWithZeros.SetTensorType(ENeuralTensorType::IntermediateInitialized);
 		// Set XSizesForGPU
 		XSizesForGPU.Copy(X.GetSizes());
 	}
diff --git a/Engine/Plugins/Experimental/NeuralNetworkInference/Source/NeuralNetworkInference/Private/NeuralTensor.cpp b/Engine/Plugins/Experimental/NeuralNetworkInference/Source/NeuralNetworkInference/Private/NeuralTensor.cpp
--- a/Engine/Plugins/Experimental/NeuralNetworkInference/Source/NeuralNetworkInference/Private/NeuralTensor.cpp
+++ b/Engine/Plugins/Experimental/NeuralNetworkInference/Source/NeuralNetworkInference/Private/NeuralTensor.cpp
@@ -179,7 +179,7 @@ FNeuralTensor::FNeuralTensor(const ENeuralDataType InDataType, const void* const
 FNeuralTensor::FNeuralTensor(const FNeuralTensor& InTensor)
 	: FNeuralTensor(InTensor.DataType, InTensor.Sizes, InTensor.Name)
 {
-	SetTensorTypeGPU(InTensor.TensorType);
+	SetTensorType(InTensor.TensorType);
 	UnderlyingUInt8ArrayData = InTensor.UnderlyingUInt8ArrayData;
 	bEnableGPU = InTensor.bEnableGPU;
 	PooledBuffer = InTensor.PooledBuffer;
@@ -191,7 +191,7 @@ FNeuralTensor& FNeuralTensor::operator=(const FNeuralTensor& InTensor)
 {
 	Name = InTensor.Name;
 	SetNumUninitialized(InTensor.Sizes, InTensor.DataType);
-	SetTensorTypeGPU(InTensor.TensorType);
+	SetTensorType(InTensor.TensorType);
 	UnderlyingUInt8ArrayData = InTensor.UnderlyingUInt8ArrayData;
 	bEnableGPU = InTensor.bEnableGPU;
 	PooledBuffer = InTensor.PooledBuffer;
@@ -205,7 +205,7 @@ FNeuralTensor::FNeuralTensor(FNeuralTensor&& InTensor)
 	: FNeuralTensor(InTensor.DataType, InTensor.Sizes, TEXT(""))
 {
 	Swap(Name, InTensor.Name);
-	SetTensorTypeGPU(InTensor.TensorType);
+	SetTensorType(InTensor.TensorType);
 	Swap(UnderlyingUInt8ArrayData, InTensor.UnderlyingUInt8ArrayData);
 	bEnableGPU = InTensor.bEnableGPU;
 	Swap(PooledBuffer, InTensor.PooledBuffer);
@@ -217,7 +217,7 @@ FNeuralTensor& FNeuralTensor::operator=(FNeuralTensor&& InTensor)
 {
 	Swap(Name, InTensor.Name);
 	SetNumUninitialized(InTensor.Sizes, InTensor.DataType);
-	SetTensorTypeGPU(InTensor.TensorType);
+	SetTensorType(InTensor.TensorType);
 	Swap(UnderlyingUInt8ArrayData, InTensor.UnderlyingUInt8ArrayData);
 	bEnableGPU = InTensor.bEnableGPU;
 	Swap(PooledBuffer, InTensor.PooledBuffer);
@@ -253,7 +253,7 @@ int64 FNeuralTensor::GetSize(const int32 InDimension) const
 	return (InDimension < GetNumberDimensions() ? Sizes[InDimension] : 1u);
 }
 
-void FNeuralTensor::SetTensorTypeGPU(const ENeuralTensorType InTensorType)
+void FNeuralTensor::SetTensorType(const ENeuralTensorType InTensorType)
 {
 	// Sanity check
 	if (PooledBuffer.IsValid() || BufferSRVRef.IsValid() || BufferUAVRef.IsValid())
diff --git a/Engine/Plugins/Experimental/NeuralNetworkInference/Source/NeuralNetworkInference/Private/UEOnly/GraphProtoToNeuralNetworkConverter.cpp b/Engine/Plugins/Experimental/NeuralNetworkInference/Source/NeuralNetworkInference/Private/UEOnly/GraphProtoToNeuralNetworkConverter.cpp
--- a/Engine/Plugins/Experimental/NeuralNetworkInference/Source/NeuralNetworkInference/Private/UEOnly/GraphProtoToNeuralNetworkConverter.cpp
+++ b/Engine/Plugins/Experimental/NeuralNetworkInference/Source/NeuralNetworkInference/Private/UEOnly/GraphProtoToNeuralNetworkConverter.cpp
@@ -253,7 +253,7 @@ bool FGraphProtoToNeuralNetworkConverter::CreateOperatorsAndEditTensorArray(TArr
 			{
 				UE_LOG(LogNeuralNetworkInference, Display, TEXT("Flipping FConvTranspose weights. This is a temporary measurement but it will keep the weights flipped for now."));
 				// Sanity check
-				if (OperatorInputTensors.Num() > 1 && OperatorInputTensors[1]->GetTensorTypeGPU() != ENeuralTensorType::Weight)
+				if (OperatorInputTensors.Num() > 1 && OperatorInputTensors[1]->GetTensorType() != ENeuralTensorType::Weight)
 				{
 					UE_LOG(LogNeuralNetworkInference, Warning,
 						TEXT("FGraphProtoToNeuralNetworkConverter::CreateOperatorsAndEditTensorArray(): For now, %s operators require a constant InputTensor[0] (i.e., a fixed weight tensor), it cannot change."), *NodeProto.OperatorType);
@@ -314,7 +314,7 @@ bool FGraphProtoToNeuralNetworkConverter::CreateOperatorsAndEditTensorArray(TArr
 						if (InOutputNameDummyIndexMap.Find(OutputTensorName))
 						{
 							OutputNameIndexMap.Add(OutputTensorName, *ExistingTensorIndex);
-							InOutTensors[*ExistingTensorIndex].SetTensorTypeGPU(ENeuralTensorType::Output);
+							InOutTensors[*ExistingTensorIndex].SetTensorType(ENeuralTensorType::Output);
 						}
 					}
 					// Subcases 2 and 3: FNeuralTensor not found, it must be defined now (i.e., intermediate tensor of the UNeuralNetworkLegacy)
@@ -352,7 +352,7 @@ bool FGraphProtoToNeuralNetworkConverter::CreateOperatorsAndEditTensorArray(TArr
 						if (InOutputNameDummyIndexMap.Find(OutputTensorName))
 						{
 							OutputNameIndexMap.Add(OutputTensorName, InlinedTensorGlobalIndex);
-							InOutTensors[InlinedTensorGlobalIndex].SetTensorTypeGPU(ENeuralTensorType::Output);
+							InOutTensors[InlinedTensorGlobalIndex].SetTensorType(ENeuralTensorType::Output);
 						}
 					}
 				}
diff --git a/Engine/Plugins/Experimental/NeuralNetworkInference/Source/NeuralNetworkInference/Public/NeuralTensor.h b/Engine/Plugins/Experimental/NeuralNetworkInference/Source/NeuralNetworkInference/Public/NeuralTensor.h
--- a/Engine/Plugins/Experimental/NeuralNetworkInference/Source/NeuralNetworkInference/Public/NeuralTensor.h
+++ b/Engine/Plugins/Experimental/NeuralNetworkInference/Source/NeuralNetworkInference/Public/NeuralTensor.h
@@ -33,23 +33,21 @@ public:
 	/**
 	 * There are several constructors of FNeuralTensor, all of them similar to each other. They set and fill the main variables of this class:
 	 *
-	 * @param InDataType It sets DataType, which defines the underlying uint8 data type of the network (float, double, int32, etc).
-	 * @param InSizes It sets Sizes and Volume. Sizes defines the dimensions of the tensor and Volume the total number of elements (mathematically
-	 *  Prod(Sizes[Index]) for all Indexes). Set to empty (or omit this argument) if memory allocation is not required or the final tensor size is
-	 *  unknown.
+	 * @param InDataType It sets DataType. @see DataType for more details.
+	 * @param InSizes It sets Sizes and Volume. Set to empty (or omit this argument) if memory allocation is not required or the final tensor size is
+	 *  unknown. @see Sizes, Volume for more details.
 	 * @param InVolume Alternative to InSizes which also sets Sizes and Volume. Using InVolume is equivalent to using
 	 *  InSizes = TArray<int64>({InVolume}). I.e., it behaves like InSizes, but assumes Sizes.Num() == 1 and sets Sizes[0] = Volume = InVolume. Set
 	 *  to 0 or a negative value if memory allocation is not required or the final size is unknown.
-	 * @param InName It sets Name, which is used for GPU debugging and the ToString() function.
-	 * @param InTensorType It sets TensorType, which is used when moving memory to the GPU.
+	 * @param InName It sets Name. @see Name for more details.
+	 * @param InTensorType It sets TensorType. @see TensorType for more details.
 	 * @param InArray Alternative to InSizes/InVolume, and represents the input data to copy from. The templated constructor that uses InArray is a
 	 *  simple (but not efficient) way of initializing a FNeuralTensor from the existing InArray TArray. It is equivalent to calling
 	 *  FNeuralTensor(..., InSizes/InVolume, ...) and then SetFromArrayCopy(InArray). It is not efficient because it makes a deep copy of the data
 	 *  of InArray. For maximum speed, rather than creating this intermediate InArray TArray, you can avoid the double copy by calling
 	 *  FNeuralTensor(..., InSizes/InVolume, ...) and then filling the tensor memory with GetData()/GetDataCasted().
 	 *
-	 * In addition, the array containing the CPU data, UnderlyingUInt8ArrayData, is pre-allocated based on the combination of InDataType and
-	 * InSizes/InVolume (if InSizes/InVolume not empty/zero): UnderlyingUInt8ArrayData.Num() == sizeof(data type) x Volume.
+	 * In addition, UnderlyingUInt8ArrayData is pre-allocated from InDataType and InSizes/InVolume. @see UnderlyingUInt8ArrayData for more details.
 	 */
 	FNeuralTensor(const ENeuralDataType InDataType = ENeuralDataType::None, const TArray<int64>& InSizes = TArray<int64>(),
 		const FString& InName = TEXT("FNeuralTensor"), const ENeuralTensorType InTensorType = ENeuralTensorType::Generic);
@@ -108,7 +106,7 @@ public:
 	 *   by any of these functions will no longer be valid (except GetUnderlyingUInt8ArrayRef()).
 	 * - These functions only modify the CPU memory, they do not synchronize CPU and GPU memory automatically. The user must do this accordingly.
 	 *
-	 * There are 2 functions to access the tensor as a TArray (but none of them can modify the tensor underlying memory):
+	 * GetArrayCopy and GetUnderlyingUInt8ArrayRef allow accessing the tensor data as a TArray, but they cannot modify the data values:
 	 * - GetArrayCopy<T> (slower but safer) returns a copy of the data as a TArray<T>. T has to be the same size than sizeof(DataType).
 	 * - GetUnderlyingUInt8ArrayRef (faster but could go out of scope) returns a reference to the underlying TArray<uint8> that contains the results.
 	 */
@@ -123,13 +121,12 @@ public:
 	 *   by any of these functions will no longer be valid (except GetUnderlyingUInt8ArrayRef()).
 	 * - These functions only modify the CPU memory, they do not synchronize CPU and GPU memory automatically. The user must do this accordingly.
 	 *
-	 * Both GetData() and GetDataCasted() return a pointer of the underlying CPU data (UnderlyingUInt8ArrayData):
+	 * GetArrayCopy and GetUnderlyingUInt8ArrayRef allow accessing and optionally modifying the tensor data as a pointer:
 	 * - GetData() returns a void pointer.
-	 * - GetDataCasted() returns a pointer casted to the desired type T.
-	 * There are 2 variations of each of those 2 functions, one constant (it does not allow to modify the pointer values) and a non-const one (which
-	 * allows modifying it).
+	 * - GetDataCasted<T>() returns a pointer casted to the desired type T.
+	 * There are 2 variations of each of these 2 functions, a constant (for read-only access) and a non-const one (which allows modifying the data).
 	 *
-	 * @return Void (GetData) or T (GetDataCasted) pointer to the FNeuralTensor data (const or non-constant).
+	 * @return Void (GetData) or T (GetDataCasted) pointer to the underlying FNeuralTensor TArray<uint8> data (const or non-constant).
 	 */
 	FORCEINLINE void* GetData();
 	FORCEINLINE const void* const GetData() const;
@@ -141,47 +138,60 @@ public:
 	FORCEINLINE const T* const GetDataCasted() const;
 
 	/**
-	 * It returns a copy of the name string.
+	 * It returns a constant reference of DataType. @see DataType for more details.
 	 */
-	FORCEINLINE const FString& GetName() const;
-
 	FORCEINLINE ENeuralDataType GetDataType() const;
 
 	/**
-	 * It returns the size of the current dimension, or 1 if InDimension > GetNumberDimensions().
+	 * It returns Sizes[InDimension], i.e., the size of the current dimension, or 1 if InDimension >= GetNumberDimensions().
+	 * @see Sizes for more details.
 	 */
 	int64 GetSize(const int32 InDimension) const;
 
+	/**
+	 * It returns a constant reference of Sizes. @see Sizes for more details.
+	 */
 	FORCEINLINE const TArray<int64>& GetSizes() const;
 
+	/**
+	 * It returns the number of dimensions, i.e., Sizes.Num(). @see Sizes for more details.
+	 */
 	FORCEINLINE int32 GetNumberDimensions() const;
 
 	/**
-	 * NumInBytes() = Num() * sizeof(type used)
+	 * It returns Volume, and it is equivalent to TArray<T>::Num().
 	 */
-	FORCEINLINE int64 NumInBytes() const;
+	FORCEINLINE int64 Num() const;
 
 	/**
-	 * Num() is analog to TArray<T>::Num().
+	 * It returns the volume in bytes. I.e., UnderlyingUInt8ArrayData.Num() or equivalently:
+	 * NumInBytes() = Num() * sizeof(type used).
 	 */
-	FORCEINLINE int64 Num() const;
+	FORCEINLINE int64 NumInBytes() const;
 
 	/**
-	 * IsEmpty() is analog to TArray<T>::IsEmpty().
+	 * It returns whether the tensor is empty (i.e., Volume is 0 and/or Sizes is empty). It is equivalent to TArray<T>::IsEmpty().
 	 */
 	FORCEINLINE bool IsEmpty() const;
 
 	/**
-	 * It returns the ENeuralTensorType (Generic, Input, Intermediate(Not)Initialized, Output, Weight, etc.).
+	 * It returns a constant reference of Name. @see Name for more details.
+	 */
+	FORCEINLINE const FString& GetName() const;
+
+	/**
+	 * It returns NeuralTensorType. @see NeuralTensorType for more details.
 	 */
-	FORCEINLINE ENeuralTensorType GetTensorTypeGPU() const;
+	FORCEINLINE ENeuralTensorType GetTensorType() const;
 
 	/**
-	 * It sets the ENeuralTensorType (Generic, Input, Intermediate(Not)Initialized, Output, Weight, etc.).
-	 * If the GPU memory was already initialized, it will also UE_LOG a warning.
+	 * It sets NeuralTensorType. @see NeuralTensorType for more details.
 	 */
-	void SetTensorTypeGPU(const ENeuralTensorType InTensorType);
+	void SetTensorType(const ENeuralTensorType InTensorType);
 
+	/**
+	 * NEXT ONE TO MODIFY
+	 */
 	FORCEINLINE void SetEnableGPU(const bool bInEnableGPU);
 
 	/**
@@ -277,8 +287,14 @@ public:
 
 protected:
 	/**
-	 * General variables and properties set by the FNeuralTensor constructor and saved on the FNeuralTensor UAsset.
-	 * @see FNeuralTensor::FNeuralTensor() for more details about each one of these members.
+	 * General variables and properties set by the FNeuralTensor constructor (or their setting functions) and saved on the FNeuralTensor UAsset.
+	 * - DataType defines the underlying uint8 data type of the network (float, double, int32, etc).
+	 * - Sizes defines the dimensions of the tensor.
+	 * - Volume defines the total number of elements of the tensor. I.e., it is the product of the size of each dimension, or mathematically
+	 *   Prod(Sizes[Index]) for all Indexes.
+	 * - Name is used for GPU debugging and the ToString() function.
+	 * - TensorType defines the type of the tensor on a network (Generic, Input, Intermediate(Not)Initialized, Output, Weight, etc.). It is currently
+	 *   only used when moving memory to the GPU.
 	 */
 	UPROPERTY(VisibleAnywhere, Category = "Neural Network Inference")
 	ENeuralDataType DataType;
@@ -292,7 +308,10 @@ protected:
 	ENeuralTensorType TensorType;
 
 private:
-	/** CPU-based variables */
+	/**
+	 * The array containing the CPU data. It is pre-allocated based on the combination of DataType and Sizes/Volume (if Sizes/Volume are not
+	 * empty/zero): UnderlyingUInt8ArrayData.Num() == sizeof(data type) x Volume.
+	 */
 	UPROPERTY()
 	TArray<uint8> UnderlyingUInt8ArrayData;
 	/**
@@ -397,11 +416,6 @@ const T* const FNeuralTensor::GetDataCasted() const
 	return (T*)UnderlyingUInt8ArrayData.GetData();
 }
 
-const FString& FNeuralTensor::GetName() const
-{
-	return Name;
-}
-
 ENeuralDataType FNeuralTensor::GetDataType() const
 {
 	return DataType;
@@ -417,14 +431,14 @@ int32 FNeuralTensor::GetNumberDimensions() const
 	return Sizes.Num();
 }
 
-int64 FNeuralTensor::NumInBytes() const
+int64 FNeuralTensor::Num() const
 {
-	return UnderlyingUInt8ArrayData.Num();
+	return Volume;
 }
 
-int64 FNeuralTensor::Num() const
+int64 FNeuralTensor::NumInBytes() const
 {
-	return Volume;
+	return UnderlyingUInt8ArrayData.Num();
 }
 
 bool FNeuralTensor::IsEmpty() const
@@ -432,7 +446,12 @@ bool FNeuralTensor::IsEmpty() const
 	return UnderlyingUInt8ArrayData.IsEmpty();
 }
 
-ENeuralTensorType FNeuralTensor::GetTensorTypeGPU() const
+const FString& FNeuralTensor::GetName() const
+{
+	return Name;
+}
+
+ENeuralTensorType FNeuralTensor::GetTensorType() const
 {
 	return TensorType;
 }
