Adding NNI documentation - NNI/Public/NeuraTensor.h even more cleaned and documented:
- Several functions removed or their arguments re-ordered to be more unified and follow the very same style.
- More functions with doxygen doc.
- Updated all classes/functions that used those modified/removed functions to keep working.
- GetPooledBuffer made constant.

#rnx

diff --git a/Engine/Plugins/Experimental/NeuralNetworkInference/Source/NeuralNetworkInference/Private/NeuralNetwork.cpp b/Engine/Plugins/Experimental/NeuralNetworkInference/Source/NeuralNetworkInference/Private/NeuralNetwork.cpp
--- a/Engine/Plugins/Experimental/NeuralNetworkInference/Source/NeuralNetworkInference/Private/NeuralNetwork.cpp
+++ b/Engine/Plugins/Experimental/NeuralNetworkInference/Source/NeuralNetworkInference/Private/NeuralNetwork.cpp
@@ -413,7 +413,7 @@ void UNeuralNetwork::InputTensorsToGPU(const TArray<int32>& InTensorIndexes)
 				for (uint32 InputTensorIndex = 0; InputTensorIndex < GetInputTensorNumber(); ++InputTensorIndex)
 				{
 					FNeuralTensor& InputTensor = GetInputTensorMutable(InputTensorIndex);
-					InputTensor.ToGPU_RenderThread(&GraphBuilder);
+					InputTensor.CPUToRDGBuilder_RenderThread(&GraphBuilder);
 				}
 			}
 			// Run for the desired input tensors
@@ -422,7 +422,7 @@ void UNeuralNetwork::InputTensorsToGPU(const TArray<int32>& InTensorIndexes)
 				for (const int32 InputTensorIndex : InTensorIndexes)
 				{
 					FNeuralTensor& InputTensor = GetInputTensorMutable(InputTensorIndex);
-					InputTensor.ToGPU_RenderThread(&GraphBuilder);
+					InputTensor.CPUToRDGBuilder_RenderThread(&GraphBuilder);
 				}
 			}
 			// Execute Render Graph
@@ -455,8 +455,8 @@ void UNeuralNetwork::OutputTensorsToCPU(const TArray<int32>& InTensorIndexes)
 				for (uint32 OutputTensorIndex = 0; OutputTensorIndex < GetOutputTensorNumber(); ++OutputTensorIndex)
 				{
 					FNeuralTensor& OutputTensor = GetOutputTensorMutable(OutputTensorIndex);
-					OutputTensor.UpdateSRVAndOrUAV_RenderThread(&GraphBuilder);
-					OutputTensor.ToCPU_RenderThread(&GraphBuilder);
+					OutputTensor.GPUToRDGBuilder_RenderThread(&GraphBuilder);
+					OutputTensor.RDGBuilderToCPU_RenderThread(&GraphBuilder);
 				}
 			}
 			// Run for the desired output tensors
@@ -465,8 +465,8 @@ void UNeuralNetwork::OutputTensorsToCPU(const TArray<int32>& InTensorIndexes)
 				for (const int32 OutputTensorIndex : InTensorIndexes)
 				{
 					FNeuralTensor& OutputTensor = GetOutputTensorMutable(OutputTensorIndex);
-					OutputTensor.UpdateSRVAndOrUAV_RenderThread(&GraphBuilder);
-					OutputTensor.ToCPU_RenderThread(&GraphBuilder);
+					OutputTensor.GPUToRDGBuilder_RenderThread(&GraphBuilder);
+					OutputTensor.RDGBuilderToCPU_RenderThread(&GraphBuilder);
 				}
 			}
 			// Execute Render Graph
diff --git a/Engine/Plugins/Experimental/NeuralNetworkInference/Source/NeuralNetworkInference/Private/NeuralNetworkImplBackEndUEAndORT.cpp b/Engine/Plugins/Experimental/NeuralNetworkInference/Source/NeuralNetworkInference/Private/NeuralNetworkImplBackEndUEAndORT.cpp
--- a/Engine/Plugins/Experimental/NeuralNetworkInference/Source/NeuralNetworkInference/Private/NeuralNetworkImplBackEndUEAndORT.cpp
+++ b/Engine/Plugins/Experimental/NeuralNetworkInference/Source/NeuralNetworkInference/Private/NeuralNetworkImplBackEndUEAndORT.cpp
@@ -778,7 +778,7 @@ bool UNeuralNetwork::FImplBackEndUEAndORT::SetTensorsFromNetwork(TArray<FNeuralT
 		for (int32 TensorIndex = 0; TensorIndex < TensorNumber; ++TensorIndex)
 		{
 			const char* TensorName = TensorNames[TensorIndex];
-			OutTensors.Emplace(FNeuralTensor(ANSI_TO_TCHAR(TensorName), InTensorGPUTypes[TensorIndex]));
+			OutTensors.Emplace(FNeuralTensor(/*NeuralDataType unknown yet*/ ENeuralDataType::None, /*Volume unknown yet*/0, ANSI_TO_TCHAR(TensorName), InTensorGPUTypes[TensorIndex]));
 		}
 	}
 	else
@@ -802,16 +802,18 @@ bool UNeuralNetwork::FImplBackEndUEAndORT::SetTensorsFromNetwork(TArray<FNeuralT
 
 #ifdef PLATFORM_WIN64
 		if (InTensorGPUTypes[TensorIndex] == ENeuralTensorType::Generic)
+#endif
 		{
 			// Pre-allocate TArray (if size is different)
-			OutTensors[TensorIndex].SetNumUninitialized(InSizes[TensorIndex], InTensorDataTypes[TensorIndex]);
+			OutTensors[TensorIndex].SetNumUninitialized(InTensorDataTypes[TensorIndex], InSizes[TensorIndex]);
 			// Link tensor with ORT blob
 			LinkTensorToONNXRuntime(OutTensors, OrtTensors, *AllocatorInfo, TensorIndex);
 		}
+#ifdef PLATFORM_WIN64
 		else if (InTensorGPUTypes[TensorIndex] == ENeuralTensorType::Input || InTensorGPUTypes[TensorIndex] == ENeuralTensorType::Output)
 		{
 			// @todo: should we remove this? It's currently used to read memory from GPU to CPU
-			OutTensors[TensorIndex].SetNumUninitialized(InSizes[TensorIndex], InTensorDataTypes[TensorIndex]);
+			OutTensors[TensorIndex].SetNumUninitialized(InTensorDataTypes[TensorIndex], InSizes[TensorIndex]);
 
 			OutTensors[TensorIndex].SetEnableGPU(true);
 
@@ -821,7 +823,7 @@ bool UNeuralNetwork::FImplBackEndUEAndORT::SetTensorsFromNetwork(TArray<FNeuralT
 			if (!OutTensors[TensorIndex].InitPooledBuffer(&D3DResource))
 			{
 				UE_LOG(LogNeuralNetworkInference, Warning,
-					TEXT("FImplBackEndUEAndORT::SetTensorsFromNetwork(): Failed to initialize pooled buffer"));
+					TEXT("FImplBackEndUEAndORT::SetTensorsFromNetwork(): Failed to initialize the pooled buffer."));
 				return false;
 			}
 
@@ -829,15 +831,10 @@ bool UNeuralNetwork::FImplBackEndUEAndORT::SetTensorsFromNetwork(TArray<FNeuralT
 			if (!LinkTensorResourceToONNXRuntime(OutTensors[TensorIndex], OrtTensors[TensorIndex], D3DResource))
 			{
 				UE_LOG(LogNeuralNetworkInference, Warning,
-					TEXT("FImplBackEndUEAndORT::SetTensorsFromNetwork(): Failed to link GPU resource to ONNX runtime"));
+					TEXT("FImplBackEndUEAndORT::SetTensorsFromNetwork(): Failed to link the GPU resource to ONNX Runtime."));
 				return false;
 			}
 		}
-#else
-		// Pre-allocate TArray (if size is different)
-		OutTensors[TensorIndex].SetNumUninitialized(InSizes[TensorIndex], InTensorDataTypes[TensorIndex]);
-		// Link tensor with ORT blob
-		LinkTensorToONNXRuntime(OutTensors, OrtTensors, *AllocatorInfo, TensorIndex);
 #endif
 	}
 
@@ -990,7 +987,7 @@ void UNeuralNetwork::FImplBackEndUEAndORT::RunSessionImpl(const ENeuralDeviceTyp
 		// Copy data to GPU (if required)
 		bool bNeedsGPUCopy = false;
 
-		for (auto& InputTensor : InputTensors)
+		for (FNeuralTensor& InputTensor : InputTensors)
 		{
 			if (InputTensor.GetTensorType() != ENeuralTensorType::Input)
 				continue;
@@ -1002,8 +999,8 @@ void UNeuralNetwork::FImplBackEndUEAndORT::RunSessionImpl(const ENeuralDeviceTyp
 					FRDGBuilder GraphBuilder(RHICmdList, RDG_EVENT_NAME("UploadTensorToGPU"));
 
 					// Set parameters
-					TRefCountPtr< FRDGPooledBuffer >& PooledBuffer = InputTensor.GetPooledBuffer();
-					FRDGBufferRef						InputBufferRef = GraphBuilder.RegisterExternalBuffer(PooledBuffer);
+					const TRefCountPtr<FRDGPooledBuffer>& PooledBuffer = InputTensor.GetPooledBuffer();
+					FRDGBufferRef InputBufferRef = GraphBuilder.RegisterExternalBuffer(PooledBuffer);
 
 					FUploadTensorParameters* UploadParameters = GraphBuilder.AllocParameters<FUploadTensorParameters>();
 
diff --git a/Engine/Plugins/Experimental/NeuralNetworkInference/Source/NeuralNetworkInference/Private/NeuralNetworkImplBackEndUEOnly.cpp b/Engine/Plugins/Experimental/NeuralNetworkInference/Source/NeuralNetworkInference/Private/NeuralNetworkImplBackEndUEOnly.cpp
--- a/Engine/Plugins/Experimental/NeuralNetworkInference/Source/NeuralNetworkInference/Private/NeuralNetworkImplBackEndUEOnly.cpp
+++ b/Engine/Plugins/Experimental/NeuralNetworkInference/Source/NeuralNetworkInference/Private/NeuralNetworkImplBackEndUEOnly.cpp
@@ -135,11 +135,11 @@ void UNeuralNetwork::FImplBackEndUEOnly::Run(FOnAsyncRunCompleted& InOutOnAsyncR
 						{
 							for (const int32 InputIndex : TensorManager.GetInputIndexes())
 							{
-								Tensors[InputIndex].ToGPU_RenderThread(&GraphBuilder);
+								Tensors[InputIndex].CPUToRDGBuilder_RenderThread(&GraphBuilder);
 							}
 							for (const int32 NonInputIndex : TensorManager.GetNonInputIndexes())
 							{
-								Tensors[NonInputIndex].UpdateSRVAndOrUAV_RenderThread(&GraphBuilder);
+								Tensors[NonInputIndex].GPUToRDGBuilder_RenderThread(&GraphBuilder);
 							}
 						}
 						// If input is already on GPU, just refresh all w.r.t. GraphBuilder
@@ -147,7 +147,7 @@ void UNeuralNetwork::FImplBackEndUEOnly::Run(FOnAsyncRunCompleted& InOutOnAsyncR
 						{
 							for (FNeuralTensor& Tensor : Tensors)
 							{
-								Tensor.UpdateSRVAndOrUAV_RenderThread(&GraphBuilder);
+								Tensor.GPUToRDGBuilder_RenderThread(&GraphBuilder);
 							}
 						}
 					}
@@ -161,7 +161,7 @@ void UNeuralNetwork::FImplBackEndUEOnly::Run(FOnAsyncRunCompleted& InOutOnAsyncR
 							// Move tensors to GPU
 							for (FNeuralTensor& Tensor : Tensors)
 							{
-								Tensor.ToGPU_RenderThread(&GraphBuilder);
+								Tensor.CPUToRDGBuilder_RenderThread(&GraphBuilder);
 							}
 						}
 						// If input is already on GPU, just refresh inputs
@@ -169,11 +169,11 @@ void UNeuralNetwork::FImplBackEndUEOnly::Run(FOnAsyncRunCompleted& InOutOnAsyncR
 						{
 							for (const int32 InputIndex : TensorManager.GetInputIndexes())
 							{
-								Tensors[InputIndex].UpdateSRVAndOrUAV_RenderThread(&GraphBuilder);
+								Tensors[InputIndex].GPUToRDGBuilder_RenderThread(&GraphBuilder);
 							}
 							for (const int32 NonInputIndex : TensorManager.GetNonInputIndexes())
 							{
-								Tensors[NonInputIndex].ToGPU_RenderThread(&GraphBuilder);
+								Tensors[NonInputIndex].CPUToRDGBuilder_RenderThread(&GraphBuilder);
 							}
 						}
 						// Operators->ToGPU()
@@ -203,7 +203,7 @@ void UNeuralNetwork::FImplBackEndUEOnly::Run(FOnAsyncRunCompleted& InOutOnAsyncR
 					{
 						for (const int32 OutputIndex : TensorManager.GetOutputIndexes())
 						{
-							Tensors[OutputIndex].ToCPU_RenderThread(&GraphBuilder);
+							Tensors[OutputIndex].RDGBuilderToCPU_RenderThread(&GraphBuilder);
 						}
 					}
 
diff --git a/Engine/Plugins/Experimental/NeuralNetworkInference/Source/NeuralNetworkInference/Private/NeuralOperators/BatchNormalizationOperator.cpp b/Engine/Plugins/Experimental/NeuralNetworkInference/Source/NeuralNetworkInference/Private/NeuralOperators/BatchNormalizationOperator.cpp
--- a/Engine/Plugins/Experimental/NeuralNetworkInference/Source/NeuralNetworkInference/Private/NeuralOperators/BatchNormalizationOperator.cpp
+++ b/Engine/Plugins/Experimental/NeuralNetworkInference/Source/NeuralNetworkInference/Private/NeuralOperators/BatchNormalizationOperator.cpp
@@ -103,7 +103,7 @@ bool FBatchNormalizationOperator::ConfigureOutputAndInternalVariablesAndSanityCh
 		FNeuralTensor& OutputTensor = GetOutputTensorNoConst();
 		if (OutputTensor.GetSizes() != InputTensors[0]->GetSizes())
 		{
-			OutputTensor.SetNumUninitialized(*InputTensors[0]);
+			OutputTensor.SetNumUninitialized(InputTensors[0]->GetDataType(), InputTensors[0]->GetSizes());
 		}
 	}
 	// Return true
diff --git a/Engine/Plugins/Experimental/NeuralNetworkInference/Source/NeuralNetworkInference/Private/NeuralOperators/ConstantOperator.cpp b/Engine/Plugins/Experimental/NeuralNetworkInference/Source/NeuralNetworkInference/Private/NeuralOperators/ConstantOperator.cpp
--- a/Engine/Plugins/Experimental/NeuralNetworkInference/Source/NeuralNetworkInference/Private/NeuralOperators/ConstantOperator.cpp
+++ b/Engine/Plugins/Experimental/NeuralNetworkInference/Source/NeuralNetworkInference/Private/NeuralOperators/ConstantOperator.cpp
@@ -129,7 +129,7 @@ void FConstantOperator::ForwardCPU()
 
 	// Fill output tensor
 	FNeuralTensor& OutputTensor = GetOutputTensorNoConst();
-	OutputTensor.SetNumUninitialized(Tensor);
+	OutputTensor.SetNumUninitialized(Tensor.GetDataType(), Tensor.GetSizes());
 	OutputTensor.SetFromUnderlyingUInt8ArrayCopy(Tensor.GetUnderlyingUInt8ArrayRef());
 }
 
diff --git a/Engine/Plugins/Experimental/NeuralNetworkInference/Source/NeuralNetworkInference/Private/NeuralOperators/ConvBaseOperator.cpp b/Engine/Plugins/Experimental/NeuralNetworkInference/Source/NeuralNetworkInference/Private/NeuralOperators/ConvBaseOperator.cpp
--- a/Engine/Plugins/Experimental/NeuralNetworkInference/Source/NeuralNetworkInference/Private/NeuralOperators/ConvBaseOperator.cpp
+++ b/Engine/Plugins/Experimental/NeuralNetworkInference/Source/NeuralNetworkInference/Private/NeuralOperators/ConvBaseOperator.cpp
@@ -716,7 +716,7 @@ bool FConvBaseOperator::FillOutputSizesAndPaddedMarginsAndSanityChecks(TArray<in
 	}
 	// Initialize OutputTensor
 	FNeuralTensor& OutputTensor = GetOutputTensorNoConst();
-	OutputTensor.SetNumUninitialized(OutputSizes, XOrXWithZeros.GetDataType());
+	OutputTensor.SetNumUninitialized(XOrXWithZeros.GetDataType(), OutputSizes);
 	// Sanity check
 	if (!FNeuralNetworkInferenceUtils::SizeSanityChecks(GetOutputTensorsConst(), 1, 1, 3, -1))
 	{
diff --git a/Engine/Plugins/Experimental/NeuralNetworkInference/Source/NeuralNetworkInference/Private/NeuralOperators/ConvTransposeOperator.cpp b/Engine/Plugins/Experimental/NeuralNetworkInference/Source/NeuralNetworkInference/Private/NeuralOperators/ConvTransposeOperator.cpp
--- a/Engine/Plugins/Experimental/NeuralNetworkInference/Source/NeuralNetworkInference/Private/NeuralOperators/ConvTransposeOperator.cpp
+++ b/Engine/Plugins/Experimental/NeuralNetworkInference/Source/NeuralNetworkInference/Private/NeuralOperators/ConvTransposeOperator.cpp
@@ -454,7 +454,7 @@ bool FConvTransposeOperator::SetAndConfigureAuxiliaryTensor()
 		// Set XWithZeros (i.e., AuxiliaryTensors[0])
 		FNeuralTensor& XWithZeros = *AuxiliaryTensors[0];
 		// Initialize XWithZeros
-		XWithZeros.SetNumUninitialized(XWithZerosSizes, X.GetDataType());
+		XWithZeros.SetNumUninitialized(X.GetDataType(), XWithZerosSizes);
 		XWithZeros.SetTo<float>(0.f);
 		XWithZeros.SetTensorType(ENeuralTensorType::IntermediateInitialized);
 		// Set XSizesForGPU
diff --git a/Engine/Plugins/Experimental/NeuralNetworkInference/Source/NeuralNetworkInference/Private/NeuralOperators/ElementWiseOperator.cpp b/Engine/Plugins/Experimental/NeuralNetworkInference/Source/NeuralNetworkInference/Private/NeuralOperators/ElementWiseOperator.cpp
--- a/Engine/Plugins/Experimental/NeuralNetworkInference/Source/NeuralNetworkInference/Private/NeuralOperators/ElementWiseOperator.cpp
+++ b/Engine/Plugins/Experimental/NeuralNetworkInference/Source/NeuralNetworkInference/Private/NeuralOperators/ElementWiseOperator.cpp
@@ -56,7 +56,7 @@ bool IElementWiseOperator::ConfigureOutputAndInternalVariablesAndSanityChecks()
 		}
 		// Set and configure output tensor
 		FNeuralTensor& OutputTensor = *OutputTensors.Last();
-		OutputTensor.SetNumUninitialized(InputTensor);
+		OutputTensor.SetNumUninitialized(InputTensor.GetDataType(), InputTensor.GetSizes());
 	}
 	// Return true
 	bIsLoaded = true;
diff --git a/Engine/Plugins/Experimental/NeuralNetworkInference/Source/NeuralNetworkInference/Private/NeuralOperators/GemmOperator.cpp b/Engine/Plugins/Experimental/NeuralNetworkInference/Source/NeuralNetworkInference/Private/NeuralOperators/GemmOperator.cpp
--- a/Engine/Plugins/Experimental/NeuralNetworkInference/Source/NeuralNetworkInference/Private/NeuralOperators/GemmOperator.cpp
+++ b/Engine/Plugins/Experimental/NeuralNetworkInference/Source/NeuralNetworkInference/Private/NeuralOperators/GemmOperator.cpp
@@ -81,7 +81,7 @@ bool FGemmOperator::ConfigureOutputAndInternalVariablesAndSanityChecks()
 	const int64 M = InputTensors[0]->GetSize(bTransA ? 1 : 0);
 	const int64 N = InputTensors[1]->GetSize(bTransB ? 0 : 1);
 	FNeuralTensor& OutputTensor = GetOutputTensorNoConst();
-	OutputTensor.SetNumUninitialized(TArray<int64>({ M, N }), InputTensors[0]->GetDataType());
+	OutputTensor.SetNumUninitialized(InputTensors[0]->GetDataType(), TArray<int64>({ M, N }));
 	// Output sanity checks
 	if (!FNeuralNetworkInferenceUtils::SizeSanityChecks(GetOutputTensorsNoConst(), 1, 1, 2, 3))
 	{
diff --git a/Engine/Plugins/Experimental/NeuralNetworkInference/Source/NeuralNetworkInference/Private/NeuralOperators/MultidirectionalBroadcastOperator.cpp b/Engine/Plugins/Experimental/NeuralNetworkInference/Source/NeuralNetworkInference/Private/NeuralOperators/MultidirectionalBroadcastOperator.cpp
--- a/Engine/Plugins/Experimental/NeuralNetworkInference/Source/NeuralNetworkInference/Private/NeuralOperators/MultidirectionalBroadcastOperator.cpp
+++ b/Engine/Plugins/Experimental/NeuralNetworkInference/Source/NeuralNetworkInference/Private/NeuralOperators/MultidirectionalBroadcastOperator.cpp
@@ -220,7 +220,7 @@ bool IMultidirectionalBroadcastOperator::ConfigureOutputAndInternalVariablesAndS
 		}
 		// Set and configure output tensor
 		FNeuralTensor& OutputTensor = *OutputTensors.Last();
-		OutputTensor.SetNumUninitialized(ShapesOutput.GetInt64(), InputTensors[0]->GetDataType());
+		OutputTensor.SetNumUninitialized(InputTensors[0]->GetDataType(), ShapesOutput.GetInt64());
 	}
 	// ShapesOutput < MAX_NUMBER_DIMENSIONS
 	if ((uint32)ShapesOutput.Num() > FMultidirectionalBroadcastCS::MAX_NUMBER_DIMENSIONS)
diff --git a/Engine/Plugins/Experimental/NeuralNetworkInference/Source/NeuralNetworkInference/Private/NeuralOperators/ReshapeOperator.cpp b/Engine/Plugins/Experimental/NeuralNetworkInference/Source/NeuralNetworkInference/Private/NeuralOperators/ReshapeOperator.cpp
--- a/Engine/Plugins/Experimental/NeuralNetworkInference/Source/NeuralNetworkInference/Private/NeuralOperators/ReshapeOperator.cpp
+++ b/Engine/Plugins/Experimental/NeuralNetworkInference/Source/NeuralNetworkInference/Private/NeuralOperators/ReshapeOperator.cpp
@@ -53,7 +53,7 @@ bool FReshapeOperator::ConfigureOutputAndInternalVariablesAndSanityChecks()
 		FNeuralTensor& OutputTensor = GetOutputTensorNoConst();
 		// Initialize OutputTensor
 		const TArray<int64> OutputShape = FReshapeOperator::Remove0sAndMinus1sFromShape(ShapeTensor, InputTensor);
-		OutputTensor.SetNumUninitialized(OutputShape, InputTensor.GetDataType());
+		OutputTensor.SetNumUninitialized(InputTensor.GetDataType(), OutputShape);
 	}
 	// Return true
 	bIsLoaded = true;
diff --git a/Engine/Plugins/Experimental/NeuralNetworkInference/Source/NeuralNetworkInference/Private/NeuralOperators/SqueezeOperator.cpp b/Engine/Plugins/Experimental/NeuralNetworkInference/Source/NeuralNetworkInference/Private/NeuralOperators/SqueezeOperator.cpp
--- a/Engine/Plugins/Experimental/NeuralNetworkInference/Source/NeuralNetworkInference/Private/NeuralOperators/SqueezeOperator.cpp
+++ b/Engine/Plugins/Experimental/NeuralNetworkInference/Source/NeuralNetworkInference/Private/NeuralOperators/SqueezeOperator.cpp
@@ -39,12 +39,12 @@ bool FSqueezeOperator::ConfigureOutputAndInternalVariablesAndSanityChecks()
 		{
 			const FNeuralTensor& AxesTensor = *InputTensors[1];
 			TArray<int64> OutputShape = FSqueezeOperator::Remove1sFromShape(InputTensor, AxesTensor);
-			OutputTensor.SetNumUninitialized(OutputShape, InputTensor.GetDataType());
+			OutputTensor.SetNumUninitialized(InputTensor.GetDataType(), OutputShape);
 		}
 		else
 		{
 			TArray<int64> OutputShape = FSqueezeOperator::Remove1sFromShape(InputTensor);
-			OutputTensor.SetNumUninitialized(OutputShape, InputTensor.GetDataType());
+			OutputTensor.SetNumUninitialized(InputTensor.GetDataType(), OutputShape);
 		}
 		
 	}
diff --git a/Engine/Plugins/Experimental/NeuralNetworkInference/Source/NeuralNetworkInference/Private/NeuralTensor.cpp b/Engine/Plugins/Experimental/NeuralNetworkInference/Source/NeuralNetworkInference/Private/NeuralTensor.cpp
--- a/Engine/Plugins/Experimental/NeuralNetworkInference/Source/NeuralNetworkInference/Private/NeuralTensor.cpp
+++ b/Engine/Plugins/Experimental/NeuralNetworkInference/Source/NeuralNetworkInference/Private/NeuralTensor.cpp
@@ -138,7 +138,7 @@ FNeuralTensor::FNeuralTensor(const ENeuralDataType InDataType, const TArray<int6
 	, bEnableGPU(false)
 {
 	// Memory allocation
-	SetNumUninitialized(InSizes, InDataType);
+	SetNumUninitialized(InDataType, InSizes);
 }
 
 FNeuralTensor::FNeuralTensor(const ENeuralDataType InDataType, const int64 InVolume, const FString& InName, const ENeuralTensorType InTensorType)
@@ -146,11 +146,6 @@ FNeuralTensor::FNeuralTensor(const ENeuralDataType InDataType, const int64 InVol
 {
 }
 
-FNeuralTensor::FNeuralTensor(const FString& InName, const ENeuralTensorType InTensorType)
-	: FNeuralTensor(ENeuralDataType::None, TArray<int64>({}), InName, InTensorType)
-{
-}
-
 
 
 /* FNeuralTensor private costructor
@@ -190,7 +185,7 @@ FNeuralTensor::FNeuralTensor(const FNeuralTensor& InTensor)
 FNeuralTensor& FNeuralTensor::operator=(const FNeuralTensor& InTensor)
 {
 	Name = InTensor.Name;
-	SetNumUninitialized(InTensor.Sizes, InTensor.DataType);
+	SetNumUninitialized(InTensor.DataType, InTensor.Sizes);
 	SetTensorType(InTensor.TensorType);
 	UnderlyingUInt8ArrayData = InTensor.UnderlyingUInt8ArrayData;
 	bEnableGPU = InTensor.bEnableGPU;
@@ -216,7 +211,7 @@ FNeuralTensor::FNeuralTensor(FNeuralTensor&& InTensor)
 FNeuralTensor& FNeuralTensor::operator=(FNeuralTensor&& InTensor)
 {
 	Swap(Name, InTensor.Name);
-	SetNumUninitialized(InTensor.Sizes, InTensor.DataType);
+	SetNumUninitialized(InTensor.DataType, InTensor.Sizes);
 	SetTensorType(InTensor.TensorType);
 	Swap(UnderlyingUInt8ArrayData, InTensor.UnderlyingUInt8ArrayData);
 	bEnableGPU = InTensor.bEnableGPU;
@@ -267,7 +262,7 @@ void FNeuralTensor::SetTensorType(const ENeuralTensorType InTensorType)
 	TensorType = InTensorType;
 }
 
-void FNeuralTensor::ToGPU_RenderThread(FRDGBuilder* InOutGraphBuilder)
+void FNeuralTensor::CPUToRDGBuilder_RenderThread(FRDGBuilder* InOutGraphBuilder)
 {
 // @todo: Volatile or not adding EBufferUsageFlags::ShaderResource causes this error:
 //Fatal error: [File:D:/P4/private_dh_research_pitt/Engine/Source/Runtime/Windows/D3D11RHI/Private/D3D11Util.cpp] [Line: 258] 
@@ -284,106 +279,43 @@ void FNeuralTensor::ToGPU_RenderThread(FRDGBuilder* InOutGraphBuilder)
 	// Call ToGPU_RenderThread with the right flags
 	if (TensorType == ENeuralTensorType::Generic)
 	{
-		return ToGPU_RenderThread(InOutGraphBuilder, EBufferUsageFlags::ShaderResource | EBufferUsageFlags::UnorderedAccess, true);
+		return CPUToRDGBuilder_RenderThread(InOutGraphBuilder, EBufferUsageFlags::ShaderResource | EBufferUsageFlags::UnorderedAccess, true);
 	}
 	else if (TensorType == ENeuralTensorType::Input)
 	{
-		return ToGPU_RenderThread(InOutGraphBuilder, EBufferUsageFlags::ShaderResource | EBufferUsageFlags::UnorderedAccess, true);
+		return CPUToRDGBuilder_RenderThread(InOutGraphBuilder, EBufferUsageFlags::ShaderResource | EBufferUsageFlags::UnorderedAccess, true);
 	}
 	else if (TensorType == ENeuralTensorType::IntermediateNotInitialized)
 	{
-		return ToGPU_RenderThread(InOutGraphBuilder, EBufferUsageFlags::ShaderResource | EBufferUsageFlags::UnorderedAccess | EBufferUsageFlags::Transient, false);
+		return CPUToRDGBuilder_RenderThread(InOutGraphBuilder, EBufferUsageFlags::ShaderResource | EBufferUsageFlags::UnorderedAccess | EBufferUsageFlags::Transient, false);
 	}
 	else if (TensorType == ENeuralTensorType::IntermediateInitialized)
 	{
-		return ToGPU_RenderThread(InOutGraphBuilder, EBufferUsageFlags::ShaderResource | EBufferUsageFlags::UnorderedAccess, true);
+		return CPUToRDGBuilder_RenderThread(InOutGraphBuilder, EBufferUsageFlags::ShaderResource | EBufferUsageFlags::UnorderedAccess, true);
 	}
 	else if (TensorType == ENeuralTensorType::Output)
 	{
-		return ToGPU_RenderThread(InOutGraphBuilder, EBufferUsageFlags::ShaderResource | EBufferUsageFlags::UnorderedAccess, false);
+		return CPUToRDGBuilder_RenderThread(InOutGraphBuilder, EBufferUsageFlags::ShaderResource | EBufferUsageFlags::UnorderedAccess, false);
 	}
 	else if (TensorType == ENeuralTensorType::Weight)
 	{
-		return ToGPU_RenderThread(InOutGraphBuilder, EBufferUsageFlags::ShaderResource | EBufferUsageFlags::Static, true);
-	}
-	UE_LOG(LogNeuralNetworkInference, Warning, TEXT("FNeuralTensor-%s::ToGPU_RenderThread(): Unimplemented TensorType = %d. Assuming ENeuralTensorType::Generic."), *Name, (int32)TensorType);
-	return ToGPU_RenderThread(InOutGraphBuilder, EBufferUsageFlags::UnorderedAccess, true);
-}
-
-void FNeuralTensor::ToGPU_RenderThread(FRDGBuilder* InOutGraphBuilder, const EBufferUsageFlags InEBufferUsageFlags, const bool bInShouldCopyFromCPU)
-{
-	// Sanity checks
-	if (!bEnableGPU || IsEmpty())
-	{
-		return;
-	}
-	checkf(!Name.IsEmpty(), TEXT("FNeuralTensor::ToGPU_RenderThread(): Name cannot be empty."));
-	checkf(IsInRenderingThread(), TEXT("FNeuralTensor-%s::ToGPU_RenderThread(): IsInRenderingThread() must be true."), *Name);
-	checkf(InOutGraphBuilder, TEXT("FNeuralTensor-%s::ToGPU_RenderThread(): InOutGraphBuilder is nullptr."), *Name);
-	// Not SRV-only and not UAV/SRV
-	checkf(EnumHasAnyFlags(InEBufferUsageFlags, EBufferUsageFlags::ShaderResource|EBufferUsageFlags::UnorderedAccess),
-		TEXT("FNeuralTensor-%s::ToGPU_RenderThread(): Unexpected case InEBufferUsageFlags = %d."), *Name, (uint32)InEBufferUsageFlags);
-	// If SRV-only and bInShouldCopyFromCPU == false
-	if (!EnumHasAnyFlags(InEBufferUsageFlags, EBufferUsageFlags::UnorderedAccess) && !bInShouldCopyFromCPU)
-	{
-		UE_LOG(LogNeuralNetworkInference, Warning, TEXT("FNeuralTensor-%s::ToGPU_RenderThread(): bInShouldCopyFromCPU must be true for SRVs (because they cannot be edited). Assumed true."), *Name);
-	}
-
-	// Create BufferRef
-	FRDGBufferRef BufferRef;
-
-	if (!PooledBuffer)
-	{
-		FRDGBufferDesc BufferDesc;
-		BufferDesc.BytesPerElement = FNeuralDataTypeUtils::GetByteSize(DataType);
-		BufferDesc.NumElements = Num();
-		BufferDesc.Usage = InEBufferUsageFlags;
-		BufferDesc.UnderlyingType = FRDGBufferDesc::EUnderlyingType::VertexBuffer;
-		
-		BufferRef = bInShouldCopyFromCPU 
-			? CreateVertexBuffer(*InOutGraphBuilder, *Name, BufferDesc, UnderlyingUInt8ArrayData.GetData(), NumInBytes(), ERDGInitialDataFlags::NoCopy)
-			: InOutGraphBuilder->CreateBuffer(BufferDesc, *Name);
-
-		// Recreate PooledBuffer for future runs
-		PooledBuffer = MakeShared<TRefCountPtr<FRDGPooledBuffer>>();
-		*PooledBuffer = InOutGraphBuilder->ConvertToExternalBuffer(BufferRef);
-	}
-	else
-	{
-		BufferRef = InOutGraphBuilder->RegisterExternalBuffer(*PooledBuffer);
-
-		InOutGraphBuilder->QueueBufferUpload(BufferRef, UnderlyingUInt8ArrayData.GetData(), NumInBytes(), ERDGInitialDataFlags::None);
-	}
-
-	// Recreate BufferSRVRef
-	if (EnumHasAnyFlags(InEBufferUsageFlags, EBufferUsageFlags::ShaderResource))
-	{
-		BufferSRVRef = MakeShared<FRDGBufferSRVRef>(InOutGraphBuilder->CreateSRV(BufferRef, FNeuralDataTypeUtils::GetPixelFormat(DataType)));
-	}
-	else
-	{
-		BufferSRVRef.Reset();
-	}
-	// Recreate BufferUAVRef
-	if (EnumHasAnyFlags(InEBufferUsageFlags, EBufferUsageFlags::UnorderedAccess))
-	{
-		BufferUAVRef = MakeShared<FRDGBufferUAVRef>(InOutGraphBuilder->CreateUAV(BufferRef, FNeuralDataTypeUtils::GetPixelFormat(DataType)));
-	}
-	else
-	{
-		BufferUAVRef.Reset();
+		return CPUToRDGBuilder_RenderThread(InOutGraphBuilder, EBufferUsageFlags::ShaderResource | EBufferUsageFlags::Static, true);
 	}
+	UE_LOG(LogNeuralNetworkInference, Warning,
+		TEXT("FNeuralTensor-%s::CPUToRDGBuilder_RenderThread(): Unimplemented TensorType = %d. Assuming ENeuralTensorType::Generic."),
+		*Name, (int32)TensorType);
+	return CPUToRDGBuilder_RenderThread(InOutGraphBuilder, EBufferUsageFlags::UnorderedAccess, true);
 }
 
-void FNeuralTensor::UpdateSRVAndOrUAV_RenderThread(FRDGBuilder* InOutGraphBuilder)
+void FNeuralTensor::GPUToRDGBuilder_RenderThread(FRDGBuilder* InOutGraphBuilder)
 {
 	if (!bEnableGPU)
 	{
 		return;
 	}
 	// Sanity checks
-	checkf(IsInRenderingThread(), TEXT("FNeuralTensor-%s::UpdateSRVAndOrUAV_RenderThread(): IsInRenderingThread() must be true."), *Name);
-	checkf(PooledBuffer.IsValid() && InOutGraphBuilder, TEXT("FNeuralTensor-%s::UpdateSRVAndOrUAV_RenderThread(): IPooledBuffer and InOutGraphBuilder cannot be nullptrs."), *Name);
+	checkf(IsInRenderingThread(), TEXT("FNeuralTensor-%s::GPUToRDGBuilder_RenderThread(): IsInRenderingThread() must be true."), *Name);
+	checkf(PooledBuffer.IsValid() && InOutGraphBuilder, TEXT("FNeuralTensor-%s::GPUToRDGBuilder_RenderThread(): IPooledBuffer and InOutGraphBuilder cannot be nullptrs."), *Name);
 	// Register BufferRef
 	FRDGBufferRef BufferRef = InOutGraphBuilder->RegisterExternalBuffer(*PooledBuffer);
 	// Recreate BufferSRVRef
@@ -398,7 +330,7 @@ void FNeuralTensor::UpdateSRVAndOrUAV_RenderThread(FRDGBuilder* InOutGraphBuilde
 	}
 }
 
-void FNeuralTensor::ToCPU_RenderThread(FRDGBuilder* InOutGraphBuilder)
+void FNeuralTensor::RDGBuilderToCPU_RenderThread(FRDGBuilder* InOutGraphBuilder)
 {
 	// Sanity checks
 	if (!bEnableGPU || IsEmpty())
@@ -407,7 +339,7 @@ void FNeuralTensor::ToCPU_RenderThread(FRDGBuilder* InOutGraphBuilder)
 	}
 	if (!FNeuralNetworkInferenceUtilsGPU::GPUSanityChecks(InOutGraphBuilder))
 	{
-		UE_LOG(LogNeuralNetworkInference, Warning, TEXT("FNeuralTensor-%s::ToCPU_RenderThread(): Sanity checks failed."), *Name);
+		UE_LOG(LogNeuralNetworkInference, Warning, TEXT("FNeuralTensor-%s::RDGBuilderToCPU_RenderThread(): Sanity checks failed."), *Name);
 		return;
 	}
 	// Read GPU memory back into CPU
@@ -424,11 +356,11 @@ void FNeuralTensor::ToCPU_RenderThread(FRDGBuilder* InOutGraphBuilder)
 	});
 }
 
-bool FNeuralTensor::InitPooledBuffer(void** NativeResource)
+bool FNeuralTensor::InitPooledBuffer(void** InOutNativeResource)
 {
 	if (!bEnableGPU)
 	{
-		UE_LOG(LogNeuralNetworkInference, Warning, TEXT("FNeuralTensor::InitPooledBuffer():bEnableGPU is false"));
+		UE_LOG(LogNeuralNetworkInference, Warning, TEXT("FNeuralTensor::InitPooledBuffer(): bEnableGPU is false."));
 		return false;
 	}
 
@@ -436,7 +368,7 @@ bool FNeuralTensor::InitPooledBuffer(void** NativeResource)
 	std::atomic<bool> bIsGraphBuilderDone(false);
 
 	ENQUEUE_RENDER_COMMAND(FNeuralTensor_InitPooledBuffer)(
-		[this, &NativeResource, &bIsGraphBuilderDone] (FRHICommandListImmediate& RHICmdList)
+		[this, &InOutNativeResource, &bIsGraphBuilderDone] (FRHICommandListImmediate& RHICmdList)
 		{
 			FRDGBuilder Builder(RHICmdList);
 
@@ -459,11 +391,11 @@ bool FNeuralTensor::InitPooledBuffer(void** NativeResource)
 			Builder.Execute();
 
 #ifdef PLATFORM_WIN64
-			if (NativeResource)
+			if (InOutNativeResource)
 			{
 				FRHIBuffer* Buffer = (*PooledBuffer)->GetRHI();
 
-				NativeResource[0] = FNeuralNetworkInferenceUtilsGPU::GetD3D12Resource(Buffer);
+				InOutNativeResource[0] = FNeuralNetworkInferenceUtilsGPU::GetD3D12Resource(Buffer);
 			}
 #endif
 
@@ -471,6 +403,7 @@ bool FNeuralTensor::InitPooledBuffer(void** NativeResource)
 		}
 	);
 
+	// Wait until the graph builder is done
 	while (!bIsGraphBuilderDone)
 	{
 		FPlatformProcess::Sleep(0.1e-3);
@@ -479,7 +412,7 @@ bool FNeuralTensor::InitPooledBuffer(void** NativeResource)
 	return true;
 }
 
-TRefCountPtr<class FRDGPooledBuffer>& FNeuralTensor::GetPooledBuffer() const
+const TRefCountPtr<FRDGPooledBuffer>& FNeuralTensor::GetPooledBuffer() const
 {
 	// Sanity checks
 	checkf(bEnableGPU, TEXT("FNeuralTensor-%s::GetPooledBuffer(): bEnableGPU must be true."), *Name);
@@ -494,7 +427,8 @@ const FRDGBufferSRVRef FNeuralTensor::GetBufferSRVRef() const
 	// Sanity checks
 	checkf(bEnableGPU, TEXT("FNeuralTensor-%s::GetBufferSRVRef(): bEnableGPU must be true."), *Name);
 	checkf(IsInRenderingThread(), TEXT("FNeuralTensor-%s::GetBufferSRVRef(): IsInRenderingThread() must be true."), *Name);
-	checkf(BufferSRVRef.IsValid(), TEXT("FNeuralTensor-%s::GetBufferSRVRef(): BufferSRVRef was a nullptr, 2 possible causes: 1) FNeuralTensor::ToGPU_RenderThread() was not called. 2) The tensor was empty."), *Name);
+	checkf(BufferSRVRef.IsValid(), TEXT("FNeuralTensor-%s::GetBufferSRVRef(): BufferSRVRef was a nullptr, 2 possible causes:"
+		" 1) FNeuralTensor::CPUToRDGBuilder_RenderThread() was not called. 2) The tensor was empty."), *Name);
 	// Return BufferSRVRef
 	return *BufferSRVRef;
 }
@@ -509,12 +443,7 @@ FRDGBufferUAVRef FNeuralTensor::GetBufferUAVRef()
 	return *BufferUAVRef;
 }
 
-void FNeuralTensor::SetNumUninitialized(const FNeuralTensor& InTensor, const bool bInAllowShrinking)
-{
-	SetNumUninitialized(InTensor.GetSizes(), InTensor.GetDataType(), bInAllowShrinking);
-}
-
-void FNeuralTensor::SetNumUninitialized(const TArray<int64>& InSizes, const ENeuralDataType InDataType, const bool bInAllowShrinking)
+void FNeuralTensor::SetNumUninitialized(const ENeuralDataType InDataType, const TArray<int64>& InSizes, const bool bInAllowShrinking)
 {
 	// Update DataType
 	if (InDataType != ENeuralDataType::None)
@@ -565,7 +494,7 @@ bool FNeuralTensor::SetFromTensorProto(const FTensorProto* const InTensorProto,
 	Name = InTensorName;
 	TensorType = InTensorType;
 	// Memory allocation
-	SetNumUninitialized(InTensorProto->Dimensions, FPrivateNeuralTensor::GetDataTypeFromTensorProtoDataType(InTensorProto->DataType));
+	SetNumUninitialized(FPrivateNeuralTensor::GetDataTypeFromTensorProtoDataType(InTensorProto->DataType), InTensorProto->Dimensions);
 
 	// RawData
 	if (!InTensorProto->RawData.IsEmpty())
@@ -899,3 +828,67 @@ bool FNeuralTensor::CheckTAndDataTypeResult(const bool bInCheckTAndDataTypeResul
 	}
 	return bInCheckTAndDataTypeResult;
 }
+
+void FNeuralTensor::CPUToRDGBuilder_RenderThread(FRDGBuilder* InOutGraphBuilder, const EBufferUsageFlags InBufferUsageFlags, const bool bInShouldCopyFromCPU)
+{
+	// Sanity checks
+	if (!bEnableGPU || IsEmpty())
+	{
+		return;
+	}
+	checkf(!Name.IsEmpty(), TEXT("FNeuralTensor::CPUToRDGBuilder_RenderThread(): Name cannot be empty."));
+	checkf(IsInRenderingThread(), TEXT("FNeuralTensor-%s::CPUToRDGBuilder_RenderThread(): IsInRenderingThread() must be true."), *Name);
+	checkf(InOutGraphBuilder, TEXT("FNeuralTensor-%s::CPUToRDGBuilder_RenderThread(): InOutGraphBuilder is nullptr."), *Name);
+	// Not SRV-only and not UAV/SRV
+	checkf(EnumHasAnyFlags(InBufferUsageFlags, EBufferUsageFlags::ShaderResource|EBufferUsageFlags::UnorderedAccess),
+		TEXT("FNeuralTensor-%s::CPUToRDGBuilder_RenderThread(): Unexpected case InBufferUsageFlags = %d."), *Name, (uint32)InBufferUsageFlags);
+	// If SRV-only and bInShouldCopyFromCPU == false
+	if (!EnumHasAnyFlags(InBufferUsageFlags, EBufferUsageFlags::UnorderedAccess) && !bInShouldCopyFromCPU)
+	{
+		UE_LOG(LogNeuralNetworkInference, Warning, TEXT("FNeuralTensor-%s::CPUToRDGBuilder_RenderThread(): bInShouldCopyFromCPU must be true for"
+			" SRVs (because they cannot be edited). Assumed true."), *Name);
+	}
+
+	// Create BufferRef
+	FRDGBufferRef BufferRef;
+	if (!PooledBuffer)
+	{
+		FRDGBufferDesc BufferDesc;
+		BufferDesc.BytesPerElement = FNeuralDataTypeUtils::GetByteSize(DataType);
+		BufferDesc.NumElements = Num();
+		BufferDesc.Usage = InBufferUsageFlags;
+		BufferDesc.UnderlyingType = FRDGBufferDesc::EUnderlyingType::VertexBuffer;
+		
+		BufferRef = bInShouldCopyFromCPU 
+			? CreateVertexBuffer(*InOutGraphBuilder, *Name, BufferDesc, UnderlyingUInt8ArrayData.GetData(), NumInBytes(), ERDGInitialDataFlags::NoCopy)
+			: InOutGraphBuilder->CreateBuffer(BufferDesc, *Name);
+
+		// Recreate PooledBuffer for future runs
+		PooledBuffer = MakeShared<TRefCountPtr<FRDGPooledBuffer>>();
+		*PooledBuffer = InOutGraphBuilder->ConvertToExternalBuffer(BufferRef);
+	}
+	else
+	{
+		BufferRef = InOutGraphBuilder->RegisterExternalBuffer(*PooledBuffer);
+		InOutGraphBuilder->QueueBufferUpload(BufferRef, UnderlyingUInt8ArrayData.GetData(), NumInBytes(), ERDGInitialDataFlags::None);
+	}
+
+	// Recreate BufferSRVRef
+	if (EnumHasAnyFlags(InBufferUsageFlags, EBufferUsageFlags::ShaderResource))
+	{
+		BufferSRVRef = MakeShared<FRDGBufferSRVRef>(InOutGraphBuilder->CreateSRV(BufferRef, FNeuralDataTypeUtils::GetPixelFormat(DataType)));
+	}
+	else
+	{
+		BufferSRVRef.Reset();
+	}
+	// Recreate BufferUAVRef
+	if (EnumHasAnyFlags(InBufferUsageFlags, EBufferUsageFlags::UnorderedAccess))
+	{
+		BufferUAVRef = MakeShared<FRDGBufferUAVRef>(InOutGraphBuilder->CreateUAV(BufferRef, FNeuralDataTypeUtils::GetPixelFormat(DataType)));
+	}
+	else
+	{
+		BufferUAVRef.Reset();
+	}
+}
diff --git a/Engine/Plugins/Experimental/NeuralNetworkInference/Source/NeuralNetworkInference/Public/NeuralTensor.h b/Engine/Plugins/Experimental/NeuralNetworkInference/Source/NeuralNetworkInference/Public/NeuralTensor.h
--- a/Engine/Plugins/Experimental/NeuralNetworkInference/Source/NeuralNetworkInference/Public/NeuralTensor.h
+++ b/Engine/Plugins/Experimental/NeuralNetworkInference/Source/NeuralNetworkInference/Public/NeuralTensor.h
@@ -21,7 +21,7 @@
  * FNeuralTensor is an auxiliary class of UNeuralNetwork which represents a tensor of the UNeuralNetwork model. It is Unreal Engine's equivalent of
  * torch.Tensor (PyTorch) or caffe::Blob (Caffe).
  *
- * Most of FNeuralTensor's functions run on the CPU, so `ToGPU_RenderThread()` must be called before running on GPU and after running any
+ * Most of FNeuralTensor's functions run on the CPU, so `CPUToRDGBuilder_RenderThread()` must be called before running on GPU and after running any
  * FNeuralTensor function that modifies the CPU memory. In addition, FNeuralTensor's CPU functions are very similar to those of TArray<T>.
  */
 USTRUCT()
@@ -53,7 +53,6 @@ public:
 		const FString& InName = TEXT("FNeuralTensor"), const ENeuralTensorType InTensorType = ENeuralTensorType::Generic);
 	FNeuralTensor(const ENeuralDataType InDataType, const int64 InVolume, const FString& InName = TEXT("FNeuralTensor"),
 		const ENeuralTensorType InTensorType = ENeuralTensorType::Generic);
-	FNeuralTensor(const FString& InName, const ENeuralTensorType InTensorType = ENeuralTensorType::Generic);
 	template <typename T>
 	FNeuralTensor(const TArray<T>& InArray, const TArray<int64>& InSizes = TArray<int64>(), const FString& InName = TEXT("FNeuralTensor"),
 		const ENeuralTensorType InTensorType = ENeuralTensorType::Generic);
@@ -159,7 +158,7 @@ public:
 	FORCEINLINE int32 GetNumberDimensions() const;
 
 	/**
-	 * It returns Volume, and it is equivalent to TArray<T>::Num().
+	 * It returns Volume, i.e., the number of elements of the tensor.
 	 */
 	FORCEINLINE int64 Num() const;
 
@@ -170,7 +169,7 @@ public:
 	FORCEINLINE int64 NumInBytes() const;
 
 	/**
-	 * It returns whether the tensor is empty (i.e., Volume is 0 and/or Sizes is empty). It is equivalent to TArray<T>::IsEmpty().
+	 * It returns whether the tensor is empty (i.e., Volume is 0 and/or Sizes is empty).
 	 */
 	FORCEINLINE bool IsEmpty() const;
 
@@ -190,43 +189,60 @@ public:
 	void SetTensorType(const ENeuralTensorType InTensorType);
 
 	/**
-	 * NEXT ONE TO MODIFY
+	 * It sets bEnableGPU. @see bEnableGPU for more details.
 	 */
 	FORCEINLINE void SetEnableGPU(const bool bInEnableGPU);
 
 	/**
-	 * It uploads/downloads the memory from/to the CPU to/from the GPU based on TensorType (which sets a preset of EBufferUsageFlags flags).
-	 * @param InEBufferUsageFlags gives the user total control over the buffer flags (and ignores the TensorType flag). This is meant to be filled
-	 * with a combination of EBufferUsageFlags values.
+	 * There are different functions to upload/download memory from/to CPU to/from GPU:
+	 * - CPUToRDGBuilder_RenderThread either uploads the memory from the CPU to the GPU (if bInShouldCopyFromCPU) or simply allocates the desired GPU
+	 *   memory. Then, it registers that GPU memory on the input FRDGBuilder.
+	 * - GPUToRDGBuilder_RenderThread implies that the memory is already on the GPU, and simply registers that GPU memory on the input FRDGBuilder.
+	 * - RDGBuilderToCPU_RenderThread downloads the memory from the GPU to the CPU. It is the opposite operation than CPUToRDGBuilder_RenderThread.
+	 *
+	 * @see CPUToRDGBuilder_RenderThread, GPUToRDGBuilder_RenderThread, RDGBuilderToCPU_RenderThread
 	 */
-	void ToGPU_RenderThread(FRDGBuilder* InOutGraphBuilder);
-	void ToGPU_RenderThread(FRDGBuilder* InOutGraphBuilder, const EBufferUsageFlags InEBufferUsageFlags, const bool bInShouldCopyFromCPU);
-	void UpdateSRVAndOrUAV_RenderThread(FRDGBuilder* InOutGraphBuilder);
-	void ToCPU_RenderThread(FRDGBuilder* InOutGraphBuilder);
+	void CPUToRDGBuilder_RenderThread(FRDGBuilder* InOutGraphBuilder);
+	void GPUToRDGBuilder_RenderThread(FRDGBuilder* InOutGraphBuilder);
+	void RDGBuilderToCPU_RenderThread(FRDGBuilder* InOutGraphBuilder);
 
 	/**
-	 * Allocate data for GPU pooled buffer.
-	 * NativeResource is a pointer to ID3D12Resource obtained from PooledBuffer and that one can be shared by DirectML execution provider
+	 * It allocates the GPU data, in particular the pooled buffer.
+	 * @param InOutNativeResource Pointer to the ID3D12Resource obtained from PooledBuffer. This resource can be shared with the DirectML execution
+	 * provider.
 	 */
-	bool InitPooledBuffer(void** NativeResource = nullptr);
+	bool InitPooledBuffer(void** InOutNativeResource = nullptr);
 
-	TRefCountPtr<FRDGPooledBuffer>& GetPooledBuffer() const;
+	/**
+	 * It returns a constant reference of PooledBuffer. @see PooledBuffer for more details.
+	 */
+	const TRefCountPtr<FRDGPooledBuffer>& GetPooledBuffer() const;
+
+	/**
+	 * It returns a const reference of BufferSRVRef. @see BufferSRVRef for more details.
+	 */
 	const FRDGBufferSRVRef GetBufferSRVRef() const;
+
+	/**
+	 * It returns a reference of BufferUAVRef. @see BufferUAVRef for more details.
+	 */
 	FRDGBufferUAVRef GetBufferUAVRef();
 
 	/**
-	 * Resize the FNeuralTensor to the desired new size.
-	 * Analog to TArray<T>::SetNumUninitialized().
-	 * @param InTensor Set Sizes and DataType from the input InTensor.
-	 * @param InVolume Set InVolume to 0 if memory allocation is not required or the final size is unknown,
-	 * SetNumUninitialized(InVolume) == SetNumUninitialized({InVolume}).
-	 * @param InSizes Set InSizes to empty if memory allocation is not required or the final size is unknown.
-	 * @param InDataType set to None means that it will maintain the previous type.
+	 * It resizes the tensor to the desired new size (i.e., modiying DataType, Sizes, Volume, UnderlyingUInt8ArrayData, etc.).
+	 * @param InDataType It sets DataType. If None, it will keep its current value. @see DataType for more details.
+	 * @param InSizes It sets Sizes and Volume. Set to empty (or omit this argument) if memory allocation is not required or the final tensor size is
+	 *  unknown. @see Sizes, Volume for more details.
+	 * @param InVolume Alternative to InSizes which also sets Sizes and Volume. Using InVolume is equivalent to using
+	 *  InSizes = TArray<int64>({InVolume}). I.e., it behaves like InSizes, but assumes Sizes.Num() == 1 and sets Sizes[0] = Volume = InVolume. Set
+	 *  to 0 or a negative value if memory allocation is not required or the final size is unknown.
+	 * @param bInAllowShrinking This boolean is used when calling UnderlyingUInt8ArrayData.SetNumUninitialized(NumInBytes(), bInAllowShrinking).
+	 * @see TArray::SetNumUninitialized for more information about bInAllowShrinking.
 	 */
-	void SetNumUninitialized(const FNeuralTensor& InTensor, const bool bInAllowShrinking = true);
-	FORCEINLINE void SetNumUninitialized(const int64 InVolume, const ENeuralDataType InDataType, const bool bInAllowShrinking = true);
-	void SetNumUninitialized(const TArray<int64>& InSizes, const ENeuralDataType InDataType, const bool bInAllowShrinking = true);
+	void SetNumUninitialized(const ENeuralDataType InDataType, const TArray<int64>& InSizes, const bool bInAllowShrinking = true);
+	FORCEINLINE void SetNumUninitialized(const ENeuralDataType InDataType, const int64 InVolume, const bool bInAllowShrinking = true);
 
+// CONTINUE HERE
 	/**
 	 * This will replace the TArray with the input one, by deeply copying the array (safer and easier to use).
 	 * For maximum performance:
@@ -343,6 +359,14 @@ private:
 	template<typename T>
 	bool CheckTAndDataTypeEquivalent() const;
 	bool CheckTAndDataTypeResult(const bool bInCheckTAndDataTypeResult, const int64 InSizeOfT) const;
+
+	/**
+	 * This function has total control over the buffer flags (and ignores the TensorType flag). Thus, it is kept private and it should only be used
+	 * by its public analog.
+	 * Depending on the value of InBufferUsageFlags, it can perform GPU allocation, CPU-to-GPU upload, apply different settings, etc.
+	 * @see CPUToRDGBuilder_RenderThread(FRDGBuilder* InOutGraphBuilder).
+	 */
+	void CPUToRDGBuilder_RenderThread(FRDGBuilder* InOutGraphBuilder, const EBufferUsageFlags InBufferUsageFlags, const bool bInShouldCopyFromCPU);
 };
 
 
@@ -461,9 +485,9 @@ void FNeuralTensor::SetEnableGPU(const bool bInEnableGPU)
 	bEnableGPU = bInEnableGPU;
 }
 
-void FNeuralTensor::SetNumUninitialized(const int64 InVolume, const ENeuralDataType InDataType, const bool bInAllowShrinking)
+void FNeuralTensor::SetNumUninitialized(const ENeuralDataType InDataType, const int64 InVolume, const bool bInAllowShrinking)
 {
-	SetNumUninitialized((InVolume > 0 ? TArray<int64>({ InVolume }) : TArray<int64>({})), InDataType, bInAllowShrinking);
+	SetNumUninitialized(InDataType, (InVolume > 0 ? TArray<int64>({ InVolume }) : TArray<int64>({})), bInAllowShrinking);
 }
 
 template<typename T>
