Adding NNI documentation:
- NNI/Public/NeuraTensor.h full cleaned and documented

#rnx

diff --git a/Engine/Plugins/Experimental/NeuralNetworkInference/Source/NeuralNetworkInference/Private/NeuralNetworkImplBackEndUEAndORT.cpp b/Engine/Plugins/Experimental/NeuralNetworkInference/Source/NeuralNetworkInference/Private/NeuralNetworkImplBackEndUEAndORT.cpp
--- a/Engine/Plugins/Experimental/NeuralNetworkInference/Source/NeuralNetworkInference/Private/NeuralNetworkImplBackEndUEAndORT.cpp
+++ b/Engine/Plugins/Experimental/NeuralNetworkInference/Source/NeuralNetworkInference/Private/NeuralNetworkImplBackEndUEAndORT.cpp
@@ -527,32 +527,32 @@ bool UNeuralNetwork::FImplBackEndUEAndORT::ConfigureMembers(const ENeuralDeviceT
 		}
 
 		// Get adapter's D3D12 device that we would like to share with DirectML execution provider
-		FD3D12DynamicRHI*			Rhi = static_cast<FD3D12DynamicRHI*>(GDynamicRHI);
-		FD3D12Adapter&				RhiAdapter = Rhi->GetAdapter(0);
-		const FD3D12AdapterDesc&	RhiAdapterDesc = RhiAdapter.GetDesc();
+		FD3D12DynamicRHI*			RHI = static_cast<FD3D12DynamicRHI*>(GDynamicRHI);
+		FD3D12Adapter&				RHIAdapter = RHI->GetAdapter(0);
+		const FD3D12AdapterDesc&	RHIAdapterDesc = RHIAdapter.GetDesc();
 
 		UE_LOG(LogNeuralNetworkInference, Display, TEXT("%d available RHI adapters. NNI using RHI adapter %s with LUID:%0x%0x."),
-			Rhi->GetNumAdapters(), RhiAdapterDesc.Desc.Description, RhiAdapterDesc.Desc.AdapterLuid.HighPart, RhiAdapterDesc.Desc.AdapterLuid.LowPart);
+			RHI->GetNumAdapters(), RHIAdapterDesc.Desc.Description, RHIAdapterDesc.Desc.AdapterLuid.HighPart, RHIAdapterDesc.Desc.AdapterLuid.LowPart);
 
-		if (Rhi->GetNumAdapters() > 1)
+		if (RHI->GetNumAdapters() > 1)
 		{
 			UE_LOG(LogNeuralNetworkInference, Display, TEXT("All available RHI adapters:"));
-			for (int32 AdapterIndex = 0; AdapterIndex < Rhi->GetNumAdapters(); ++AdapterIndex)
+			for (int32 AdapterIndex = 0; AdapterIndex < RHI->GetNumAdapters(); ++AdapterIndex)
 			{
-				const FD3D12AdapterDesc& CurrAdapterDesc = Rhi->GetAdapter(AdapterIndex).GetDesc();
+				const FD3D12AdapterDesc& CurrAdapterDesc = RHI->GetAdapter(AdapterIndex).GetDesc();
 				UE_LOG(LogNeuralNetworkInference, Display, TEXT("  - Adapter [%d] Name:%s with LUID:%0x%0x."),
 					AdapterIndex, CurrAdapterDesc.Desc.Description, CurrAdapterDesc.Desc.AdapterLuid.HighPart, CurrAdapterDesc.Desc.AdapterLuid.LowPart);
 			}
 		}
 
-		if (Rhi->GetNumAdapters() > 1 || RhiAdapterDesc.NumDeviceNodes > 1)
+		if (RHI->GetNumAdapters() > 1 || RHIAdapterDesc.NumDeviceNodes > 1)
 		{
 			UE_LOG(LogNeuralNetworkInference, Warning,
 				TEXT("FImplBackEndUEAndORT::ConfigureMembers(): There are multiple (%d) adapters and/or multiple (%d) devices, NNI is currently using only one adapter."),
-				Rhi->GetNumAdapters(), RhiAdapterDesc.NumDeviceNodes);
+				RHI->GetNumAdapters(), RHIAdapterDesc.NumDeviceNodes);
 		}
 
-		ID3D12Device* NativeDevice = RhiAdapter.GetD3DDevice();
+		ID3D12Device* NativeDevice = RHIAdapter.GetD3DDevice();
 
 		// Make sure that we have one DMLDevice per D3D12 device
 		IDMLDevice* DmlDevice = FPrivateImplBackEndUEAndORT::GetDMLDeviceThreadSafe(NativeDevice);
@@ -565,7 +565,7 @@ bool UNeuralNetwork::FImplBackEndUEAndORT::ConfigureMembers(const ENeuralDeviceT
 
 		// Get a ID3D12CommandQueue as well
 		// TODO: Should we create our own queue?
-		ID3D12CommandQueue* NativeCmdQ = Rhi->RHIGetD3DCommandQueue();
+		ID3D12CommandQueue* NativeCmdQ = RHI->RHIGetD3DCommandQueue();
 
 		// ORT GPU (Direct ML)
 		SessionOptions->SetGraphOptimizationLevel(GraphOptimizationLevel::ORT_ENABLE_ALL); // ORT_ENABLE_ALL, ORT_ENABLE_EXTENDED, ORT_ENABLE_BASIC, ORT_DISABLE_ALL
@@ -876,7 +876,7 @@ void UNeuralNetwork::FImplBackEndUEAndORT::LinkTensorToONNXRuntime(TArray<FNeura
 }
 
 #ifdef PLATFORM_WIN64
-bool UNeuralNetwork::FImplBackEndUEAndORT::LinkTensorResourceToONNXRuntime(FNeuralTensor& InOutTensor, Ort::Value& InOutOrtTensor, void* D3DResource)
+bool UNeuralNetwork::FImplBackEndUEAndORT::LinkTensorResourceToONNXRuntime(FNeuralTensor& InOutTensor, Ort::Value& InOutOrtTensor, void* InOutD3DResource)
 {
 	if (!DmlApi)
 	{
@@ -886,7 +886,7 @@ bool UNeuralNetwork::FImplBackEndUEAndORT::LinkTensorResourceToONNXRuntime(FNeur
 
 	void* DmlGPUAllocation = nullptr;
 	
-	DmlApi->CreateGPUAllocationFromD3DResource(reinterpret_cast<ID3D12Resource*>(D3DResource), &DmlGPUAllocation);
+	DmlApi->CreateGPUAllocationFromD3DResource(reinterpret_cast<ID3D12Resource*>(InOutD3DResource), &DmlGPUAllocation);
 
 	if (!DmlGPUAllocation)
 	{
@@ -989,46 +989,48 @@ void UNeuralNetwork::FImplBackEndUEAndORT::RunSessionImpl(const ENeuralDeviceTyp
 
 		for (FNeuralTensor& InputTensor : InputTensors)
 		{
-			if (InputTensor.GetTensorType() != ENeuralTensorType::Input)
-				continue;
-
-			bNeedsGPUCopy = true;
+// @todo: FIX!
+// check(InputTensor.GetTensorType() == ENeuralTensorType::Input, TEXT("All InputTensors should be ENeuralTensorType::Input, something is wrong"));
+			if (InputTensor.GetTensorType() == ENeuralTensorType::Input)
+			{
+				bNeedsGPUCopy = true;
 
-			ENQUEUE_RENDER_COMMAND(UploadTensorToGPU)([this, InputTensor](FRHICommandListImmediate& RHICmdList)
-				{
-					FRDGBuilder GraphBuilder(RHICmdList, RDG_EVENT_NAME("UploadTensorToGPU"));
+				ENQUEUE_RENDER_COMMAND(UploadTensorToGPU)([InputTensor](FRHICommandListImmediate& RHICmdList)
+					{
+						FRDGBuilder GraphBuilder(RHICmdList, RDG_EVENT_NAME("UploadTensorToGPU"));
 
-					// Set parameters
-					const TRefCountPtr<FRDGPooledBuffer>& PooledBuffer = InputTensor.GetPooledBuffer();
-					FRDGBufferRef InputBufferRef = GraphBuilder.RegisterExternalBuffer(PooledBuffer);
+						// Set parameters
+						const TRefCountPtr<FRDGPooledBuffer>& PooledBuffer = InputTensor.GetPooledBuffer();
+						FRDGBufferRef InputBufferRef = GraphBuilder.RegisterExternalBuffer(PooledBuffer);
 
-					FUploadTensorParameters* UploadParameters = GraphBuilder.AllocParameters<FUploadTensorParameters>();
+						FUploadTensorParameters* UploadParameters = GraphBuilder.AllocParameters<FUploadTensorParameters>();
 
-					UploadParameters->Input = InputBufferRef;
+						UploadParameters->Input = InputBufferRef;
 
-					GraphBuilder.AddPass(
-						RDG_EVENT_NAME("NNI:UploadTensor:%s", *InputTensor.GetName()),
-						FUploadTensorParameters::FTypeInfo::GetStructMetadata(),
-						UploadParameters,
-						ERDGPassFlags::Copy | ERDGPassFlags::NeverCull,
-						[this, UploadParameters](FRHICommandListImmediate& RHICmdList)
-						{
-							FRHIBuffer* InputBuffer = UploadParameters->Input->GetRHI();
+						GraphBuilder.AddPass(
+							RDG_EVENT_NAME("UNeuralNetwork-UploadTensor-%s", *InputTensor.GetName()),
+							FUploadTensorParameters::FTypeInfo::GetStructMetadata(),
+							UploadParameters,
+							ERDGPassFlags::Copy | ERDGPassFlags::NeverCull,
+							[UploadParameters](FRHICommandListImmediate& RHICmdList)
+							{
+								FRHIBuffer* InputBuffer = UploadParameters->Input->GetRHI();
 
-							// NOTE: We're using UAVMask to trigger the UAV barrier in RDG
-							RHICmdList.Transition(FRHITransitionInfo(InputBuffer, ERHIAccess::CopyDest, ERHIAccess::UAVMask));
-						}
-					);
+								// NOTE: We're using UAVMask to trigger the UAV barrier in RDG
+								RHICmdList.Transition(FRHITransitionInfo(InputBuffer, ERHIAccess::CopyDest, ERHIAccess::UAVMask));
+							}
+						);
 
-					GraphBuilder.Execute();
-				});
+						GraphBuilder.Execute();
+					});
+			}
 		}
 
 		if (bNeedsGPUCopy)
 		{
 			ENQUEUE_RENDER_COMMAND(FlushUploadTensorToGPU)([this](FRHICommandListImmediate& RHICmdList)
 				{
-					FRDGBuilder GraphBuilder(RHICmdList, RDG_EVENT_NAME("NNI:FlushUploadTensorsToGPU"));
+					FRDGBuilder GraphBuilder(RHICmdList, RDG_EVENT_NAME("FImplBackEndUEAndORT-FlushUploadTensorsToGPU"));
 
 					RHICmdList.SubmitCommandsHint();
 					RHICmdList.ImmediateFlush(EImmediateFlushType::FlushRHIThread);
@@ -1041,7 +1043,7 @@ void UNeuralNetwork::FImplBackEndUEAndORT::RunSessionImpl(const ENeuralDeviceTyp
 		}
 
 		// Profiling init
-		FNNIGPUProfiler::Instance()->EventBegin("NNI:SessionRun");
+		FNNIGPUProfiler::Instance()->EventBegin("FImplBackEndUEAndORT-SessionRun");
 	}
 
 	// Actual ORT network run
diff --git a/Engine/Plugins/Experimental/NeuralNetworkInference/Source/NeuralNetworkInference/Private/NeuralNetworkImplBackEndUEAndORT.h b/Engine/Plugins/Experimental/NeuralNetworkInference/Source/NeuralNetworkInference/Private/NeuralNetworkImplBackEndUEAndORT.h
--- a/Engine/Plugins/Experimental/NeuralNetworkInference/Source/NeuralNetworkInference/Private/NeuralNetworkImplBackEndUEAndORT.h
+++ b/Engine/Plugins/Experimental/NeuralNetworkInference/Source/NeuralNetworkInference/Private/NeuralNetworkImplBackEndUEAndORT.h
@@ -124,7 +124,7 @@ private:
 		const int32 InTensorIndex);
 
 #ifdef PLATFORM_WIN64
-	bool LinkTensorResourceToONNXRuntime(FNeuralTensor& InOutTensor, Ort::Value& InOutOrtTensor, void* D3DResource);
+	bool LinkTensorResourceToONNXRuntime(FNeuralTensor& InOutTensor, Ort::Value& InOutOrtTensor, void* InOutD3DResource);
 #endif
 
 	void RunSessionAsync(const ENeuralDeviceType InDeviceType, const ENeuralDeviceType InInputDeviceType, const ENeuralDeviceType InOutputDeviceType);
diff --git a/Engine/Plugins/Experimental/NeuralNetworkInference/Source/NeuralNetworkInference/Private/NeuralTensor.cpp b/Engine/Plugins/Experimental/NeuralNetworkInference/Source/NeuralNetworkInference/Private/NeuralTensor.cpp
--- a/Engine/Plugins/Experimental/NeuralNetworkInference/Source/NeuralNetworkInference/Private/NeuralTensor.cpp
+++ b/Engine/Plugins/Experimental/NeuralNetworkInference/Source/NeuralNetworkInference/Private/NeuralTensor.cpp
@@ -364,6 +364,13 @@ bool FNeuralTensor::InitPooledBuffer(void** InOutNativeResource)
 		return false;
 	}
 
+	// Recreate PooledBuffer for future runs
+	PooledBuffer = MakeShared<TRefCountPtr<FRDGPooledBuffer>>();
+
+	// NOTE: SRV and UAV can be used by callers (for example MLDeformer), look at the GetPooledBuffer()
+	BufferSRVRef.Reset();
+	BufferUAVRef.Reset();
+
 	// Make sure that we wait for Graph builder to finish
 	std::atomic<bool> bIsGraphBuilderDone(false);
 
@@ -381,20 +388,14 @@ bool FNeuralTensor::InitPooledBuffer(void** InOutNativeResource)
 			FRDGBufferRef BufferRef = Builder.CreateBuffer(BufferDesc, *GetName());
 
 			// Recreate PooledBuffer for future runs
-			PooledBuffer = MakeShared<TRefCountPtr<FRDGPooledBuffer>>();
 			*PooledBuffer = Builder.ConvertToExternalBuffer(BufferRef);
 
-			// NOTE: SRV and UAV can be used by callers (for example MLDeformer), look at the GetPooledBuffer()
-			BufferSRVRef.Reset();
-			BufferUAVRef.Reset();
-
 			Builder.Execute();
 
 #ifdef PLATFORM_WIN64
 			if (InOutNativeResource)
 			{
 				FRHIBuffer* Buffer = (*PooledBuffer)->GetRHI();
-
 				InOutNativeResource[0] = FNeuralNetworkInferenceUtilsGPU::GetD3D12Resource(Buffer);
 			}
 #endif
