Adding NNI documentation:
- NNI/Public/NeuraTensor.h further cleaned and documented
#rnx

diff --git a/Engine/Plugins/Experimental/NeuralNetworkInference/Source/NeuralNetworkInference/Private/NeuralNetworkImplBackEndUEAndORT.cpp b/Engine/Plugins/Experimental/NeuralNetworkInference/Source/NeuralNetworkInference/Private/NeuralNetworkImplBackEndUEAndORT.cpp
--- a/Engine/Plugins/Experimental/NeuralNetworkInference/Source/NeuralNetworkInference/Private/NeuralNetworkImplBackEndUEAndORT.cpp
+++ b/Engine/Plugins/Experimental/NeuralNetworkInference/Source/NeuralNetworkInference/Private/NeuralNetworkImplBackEndUEAndORT.cpp
@@ -628,7 +628,7 @@ bool UNeuralNetwork::FImplBackEndUEAndORT::ConfigureTensors(TArray<FNeuralTensor
 	TArray<const char*> TensorNames;
 	TArray<ENeuralDataType> TensorDataTypes;
 	TArray<TArray<int64>> TensorSizes;
-	TArray<ENeuralTensorType> TensorGPUTypes;
+	TArray<ENeuralTensorType> TensorTypes;
 
 	const uint32 NumberTensors = bIsInput ? Session->GetInputCount() : Session->GetOutputCount();
 	if (OutAreInputTensorSizesVariable)
@@ -710,34 +710,34 @@ bool UNeuralNetwork::FImplBackEndUEAndORT::ConfigureTensors(TArray<FNeuralTensor
 		// @todo: Should caller specify tensor GPU type?
 		// Input/Output tensor GPU type means that data should not be copied from CPU
 		// Generic means that data is on the CPU and it will be copied to GPU
-		ENeuralTensorType TensorGPUType;
+		ENeuralTensorType TensorType;
 
 		if (InDeviceType == ENeuralDeviceType::GPU)
 		{
 			if (bIsInput)
 			{
-				TensorGPUType = InInputDeviceType == ENeuralDeviceType::GPU ? ENeuralTensorType::Input : ENeuralTensorType::Generic;
+				TensorType = InInputDeviceType == ENeuralDeviceType::GPU ? ENeuralTensorType::Input : ENeuralTensorType::Generic;
 			}
 			else
 			{
-				TensorGPUType = InOutputDeviceType == ENeuralDeviceType::GPU ? ENeuralTensorType::Output : ENeuralTensorType::Generic;
+				TensorType = InOutputDeviceType == ENeuralDeviceType::GPU ? ENeuralTensorType::Output : ENeuralTensorType::Generic;
 			}
 		}
 		else
 		{
-			TensorGPUType = ENeuralTensorType::Generic;
+			TensorType = ENeuralTensorType::Generic;
 		}
 
-		TensorGPUTypes.Push(TensorGPUType);
+		TensorTypes.Push(TensorType);
 
 		CurrentTypeInfo.release();
 	}
 
-	return SetTensorsFromNetwork(OutTensors, TensorNames, TensorDataTypes, TensorSizes, TensorGPUTypes, bIsInput);
+	return SetTensorsFromNetwork(OutTensors, TensorNames, TensorDataTypes, TensorSizes, TensorTypes, bIsInput);
 }
 
 bool UNeuralNetwork::FImplBackEndUEAndORT::SetTensorsFromNetwork(TArray<FNeuralTensor>& OutTensors, TArray<const char*>& InTensorNames,
-	TArray<ENeuralDataType>& InTensorDataTypes, TArray<TArray<int64>>& InSizes, TArray<ENeuralTensorType>& InTensorGPUTypes, const bool bIsInput)
+	TArray<ENeuralDataType>& InTensorDataTypes, TArray<TArray<int64>>& InSizes, TArray<ENeuralTensorType>& InTensorTypes, const bool bIsInput)
 {
 	const int32 TensorNumber = InTensorNames.Num();
 	if (InTensorDataTypes.Num() != TensorNumber || InSizes.Num() != TensorNumber)
@@ -778,14 +778,14 @@ bool UNeuralNetwork::FImplBackEndUEAndORT::SetTensorsFromNetwork(TArray<FNeuralT
 		for (int32 TensorIndex = 0; TensorIndex < TensorNumber; ++TensorIndex)
 		{
 			const char* TensorName = TensorNames[TensorIndex];
-			OutTensors.Emplace(FNeuralTensor(/*NeuralDataType unknown yet*/ ENeuralDataType::None, /*Volume unknown yet*/0, ANSI_TO_TCHAR(TensorName), InTensorGPUTypes[TensorIndex]));
+			OutTensors.Emplace(FNeuralTensor(/*NeuralDataType unknown yet*/ ENeuralDataType::None, /*Volume unknown yet*/0, ANSI_TO_TCHAR(TensorName), InTensorTypes[TensorIndex]));
 		}
 	}
 	else
 	{
 		for (int32 TensorIndex = 0; TensorIndex < TensorNumber; ++TensorIndex)
 		{
-			OutTensors[TensorIndex].SetTensorType(InTensorGPUTypes[TensorIndex]);
+			OutTensors[TensorIndex].SetTensorType(InTensorTypes[TensorIndex]);
 		}
 	}
 
@@ -801,7 +801,7 @@ bool UNeuralNetwork::FImplBackEndUEAndORT::SetTensorsFromNetwork(TArray<FNeuralT
 		}
 
 #ifdef PLATFORM_WIN64
-		if (InTensorGPUTypes[TensorIndex] == ENeuralTensorType::Generic)
+		if (InTensorTypes[TensorIndex] == ENeuralTensorType::Generic)
 #endif
 		{
 			// Pre-allocate TArray (if size is different)
@@ -810,7 +810,7 @@ bool UNeuralNetwork::FImplBackEndUEAndORT::SetTensorsFromNetwork(TArray<FNeuralT
 			LinkTensorToONNXRuntime(OutTensors, OrtTensors, *AllocatorInfo, TensorIndex);
 		}
 #ifdef PLATFORM_WIN64
-		else if (InTensorGPUTypes[TensorIndex] == ENeuralTensorType::Input || InTensorGPUTypes[TensorIndex] == ENeuralTensorType::Output)
+		else if (InTensorTypes[TensorIndex] == ENeuralTensorType::Input || InTensorTypes[TensorIndex] == ENeuralTensorType::Output)
 		{
 			// @todo: should we remove this? It's currently used to read memory from GPU to CPU
 			OutTensors[TensorIndex].SetNumUninitialized(InTensorDataTypes[TensorIndex], InSizes[TensorIndex]);
@@ -880,7 +880,7 @@ bool UNeuralNetwork::FImplBackEndUEAndORT::LinkTensorResourceToONNXRuntime(FNeur
 {
 	if (!DmlApi)
 	{
-		UE_LOG(LogNeuralNetworkInference, Warning, TEXT("FImplBackEndUEAndORT::LinkTensorResourceToONNXRuntime(): DmlGPUAllocator is not valid"));
+		UE_LOG(LogNeuralNetworkInference, Warning, TEXT("FImplBackEndUEAndORT::LinkTensorResourceToONNXRuntime(): DmlGPUAllocator is not valid."));
 		return false;
 	}
 
@@ -890,7 +890,7 @@ bool UNeuralNetwork::FImplBackEndUEAndORT::LinkTensorResourceToONNXRuntime(FNeur
 
 	if (!DmlGPUAllocation)
 	{
-		UE_LOG(LogNeuralNetworkInference, Warning, TEXT("FImplBackEndUEAndORT::LinkTensorResourceToONNXRuntime(): DmlGPUAllocation is NULL"));
+		UE_LOG(LogNeuralNetworkInference, Warning, TEXT("FImplBackEndUEAndORT::LinkTensorResourceToONNXRuntime(): DmlGPUAllocation is nullptr."));
 		return false;
 	}
 
