Fixed minor conceptual bug in NNI (not a bug in the results but instead a conceptial bug that could result in potential future busg for using enums for different purposes of what they were intended for):
- Removed OutputDeviceType almost completely from FBackEndUEAndORT because it is not used anywhere.
- ENeuralTensorType TensorType is now being properly set to either ENeuralTensorType::Input or ENeuralTensorType::Output (never to Generic).
- Run() now uses ENeuralDeviceType InputDeviceType to know whether to do the initial GPU sharing, no need to use ENeuralTensorType's workaround anymore.

#rnx

diff --git a/Engine/Plugins/Experimental/NeuralNetworkInference/Source/NeuralNetworkInference/Private/NeuralNetwork.cpp b/Engine/Plugins/Experimental/NeuralNetworkInference/Source/NeuralNetworkInference/Private/NeuralNetwork.cpp
--- a/Engine/Plugins/Experimental/NeuralNetworkInference/Source/NeuralNetworkInference/Private/NeuralNetwork.cpp
+++ b/Engine/Plugins/Experimental/NeuralNetworkInference/Source/NeuralNetworkInference/Private/NeuralNetwork.cpp
@@ -492,7 +492,7 @@ void UNeuralNetwork::Run()
 	// UEAndORT
 	if (BackEndForCurrentPlatform == ENeuralBackEnd::UEAndORT)
 	{
-		ImplBackEndUEAndORT->Run(SynchronousMode, DeviceType, InputDeviceType, OutputDeviceType);
+		ImplBackEndUEAndORT->Run(SynchronousMode, DeviceType, InputDeviceType);
 	}
 	// UEOnly
 	else if (BackEndForCurrentPlatform == ENeuralBackEnd::UEOnly)
diff --git a/Engine/Plugins/Experimental/NeuralNetworkInference/Source/NeuralNetworkInference/Private/NeuralNetworkImplBackEndUEAndORT.cpp b/Engine/Plugins/Experimental/NeuralNetworkInference/Source/NeuralNetworkInference/Private/NeuralNetworkImplBackEndUEAndORT.cpp
--- a/Engine/Plugins/Experimental/NeuralNetworkInference/Source/NeuralNetworkInference/Private/NeuralNetworkImplBackEndUEAndORT.cpp
+++ b/Engine/Plugins/Experimental/NeuralNetworkInference/Source/NeuralNetworkInference/Private/NeuralNetworkImplBackEndUEAndORT.cpp
@@ -236,13 +236,12 @@ UNeuralNetwork::FImplBackEndUEAndORT::FNeuralNetworkAsyncTask::FNeuralNetworkAsy
 {}
 
 void UNeuralNetwork::FImplBackEndUEAndORT::FNeuralNetworkAsyncTask::SetRunSessionArgs(const ENeuralSynchronousMode InSyncMode, const ENeuralDeviceType InDeviceType,
-	const ENeuralDeviceType InInputDeviceType, const ENeuralDeviceType InOutputDeviceType)
+	const ENeuralDeviceType InInputDeviceType)
 {
 	const FScopeLock ResourcesLock(&BackEnd->ResoucesCriticalSection);
 	SyncMode = InSyncMode;
 	DeviceType = InDeviceType;
 	InputDeviceType = InInputDeviceType;
-	OutputDeviceType = InOutputDeviceType;
 }
 
 void UNeuralNetwork::FImplBackEndUEAndORT::FNeuralNetworkAsyncTask::DoWork()
@@ -250,11 +249,11 @@ void UNeuralNetwork::FImplBackEndUEAndORT::FNeuralNetworkAsyncTask::DoWork()
 	const FRedirectCoutAndCerrToUeLog RedirectCoutAndCerrToUeLog;
 	if (SyncMode == ENeuralSynchronousMode::Asynchronous)
 	{
-		BackEnd->RunSessionAsync(DeviceType, InputDeviceType, OutputDeviceType);
+		BackEnd->RunSessionAsync(DeviceType, InputDeviceType);
 	}
 	else
 	{
-		BackEnd->RunSessionSync(DeviceType, InputDeviceType, OutputDeviceType);
+		BackEnd->RunSessionSync(DeviceType, InputDeviceType);
 	}
 }
 #endif //WITH_UE_AND_ORT_SUPPORT
@@ -422,7 +421,7 @@ bool UNeuralNetwork::FImplBackEndUEAndORT::Load(TSharedPtr<FImplBackEndUEAndORT>
 }
 
 void UNeuralNetwork::FImplBackEndUEAndORT::Run(const ENeuralSynchronousMode InSynchronousMode, const ENeuralDeviceType InDeviceType,
-	const ENeuralDeviceType InInputDeviceType, const ENeuralDeviceType InOutputDeviceType)
+	const ENeuralDeviceType InInputDeviceType)
 {
 #ifdef WITH_UE_AND_ORT_SUPPORT
 #if WITH_EDITOR
@@ -430,7 +429,7 @@ void UNeuralNetwork::FImplBackEndUEAndORT::Run(const ENeuralSynchronousMode InSy
 #endif //WITH_EDITOR
 	{
 		EnsureAsyncTaskCompletion();
-		NeuralNetworkAsyncTask->GetTask().SetRunSessionArgs(InSynchronousMode, InDeviceType, InInputDeviceType, InOutputDeviceType);
+		NeuralNetworkAsyncTask->GetTask().SetRunSessionArgs(InSynchronousMode, InDeviceType, InInputDeviceType);
 
 		// Run UNeuralNetwork
 		if (InSynchronousMode == ENeuralSynchronousMode::Synchronous)
@@ -707,37 +706,19 @@ bool UNeuralNetwork::FImplBackEndUEAndORT::ConfigureTensors(TArray<FNeuralTensor
 		}
 		TensorSizes.Push(CurrentTensorSizes);
 
-		// @todo: Should caller specify tensor GPU type?
-		// Input/Output tensor GPU type means that data should not be copied from CPU
-		// Generic means that data is on the CPU and it will be copied to GPU
-		ENeuralTensorType TensorType;
-
-		if (InDeviceType == ENeuralDeviceType::GPU)
-		{
-			if (bIsInput)
-			{
-				TensorType = InInputDeviceType == ENeuralDeviceType::GPU ? ENeuralTensorType::Input : ENeuralTensorType::Generic;
-			}
-			else
-			{
-				TensorType = InOutputDeviceType == ENeuralDeviceType::GPU ? ENeuralTensorType::Output : ENeuralTensorType::Generic;
-			}
-		}
-		else
-		{
-			TensorType = ENeuralTensorType::Generic;
-		}
-
+		// Set TensorType
+		const ENeuralTensorType TensorType = (bIsInput ? ENeuralTensorType::Input : ENeuralTensorType::Output);
 		TensorTypes.Push(TensorType);
 
 		CurrentTypeInfo.release();
 	}
 
-	return SetTensorsFromNetwork(OutTensors, TensorNames, TensorDataTypes, TensorSizes, TensorTypes, bIsInput);
+	return SetTensorsFromNetwork(OutTensors, TensorNames, TensorDataTypes, TensorSizes, TensorTypes, bIsInput, InInputDeviceType, InOutputDeviceType);
 }
 
 bool UNeuralNetwork::FImplBackEndUEAndORT::SetTensorsFromNetwork(TArray<FNeuralTensor>& OutTensors, TArray<const char*>& InTensorNames,
-	TArray<ENeuralDataType>& InTensorDataTypes, TArray<TArray<int64>>& InSizes, TArray<ENeuralTensorType>& InTensorTypes, const bool bIsInput)
+	TArray<ENeuralDataType>& InTensorDataTypes, TArray<TArray<int64>>& InSizes, TArray<ENeuralTensorType>& InTensorTypes, const bool bIsInput,
+	const ENeuralDeviceType InInputDeviceType, const ENeuralDeviceType InOutputDeviceType)
 {
 	const int32 TensorNumber = InTensorNames.Num();
 	if (InTensorDataTypes.Num() != TensorNumber || InSizes.Num() != TensorNumber)
@@ -801,16 +782,7 @@ bool UNeuralNetwork::FImplBackEndUEAndORT::SetTensorsFromNetwork(TArray<FNeuralT
 		}
 
 #ifdef PLATFORM_WIN64
-		if (InTensorTypes[TensorIndex] == ENeuralTensorType::Generic)
-#endif
-		{
-			// Pre-allocate TArray (if size is different)
-			OutTensors[TensorIndex].SetNumUninitialized(InTensorDataTypes[TensorIndex], InSizes[TensorIndex]);
-			// Link tensor with ORT blob
-			LinkTensorToONNXRuntime(OutTensors, OrtTensors, *AllocatorInfo, TensorIndex);
-		}
-#ifdef PLATFORM_WIN64
-		else if (InTensorTypes[TensorIndex] == ENeuralTensorType::Input || InTensorTypes[TensorIndex] == ENeuralTensorType::Output)
+		if ((bIsInput && InInputDeviceType == ENeuralDeviceType::GPU) || (!bIsInput && InOutputDeviceType == ENeuralDeviceType::GPU))
 		{
 			// @todo: should we remove this? It's currently used to read memory from GPU to CPU
 			OutTensors[TensorIndex].SetNumUninitialized(InTensorDataTypes[TensorIndex], InSizes[TensorIndex]);
@@ -835,7 +807,14 @@ bool UNeuralNetwork::FImplBackEndUEAndORT::SetTensorsFromNetwork(TArray<FNeuralT
 				return false;
 			}
 		}
+		else
 #endif
+		{
+			// Pre-allocate TArray (if size is different)
+			OutTensors[TensorIndex].SetNumUninitialized(InTensorDataTypes[TensorIndex], InSizes[TensorIndex]);
+			// Link tensor with ORT blob
+			LinkTensorToONNXRuntime(OutTensors, OrtTensors, *AllocatorInfo, TensorIndex);
+		}
 	}
 
 	return true;
@@ -932,12 +911,11 @@ bool UNeuralNetwork::FImplBackEndUEAndORT::LinkTensorResourceToONNXRuntime(FNeur
 }
 #endif
 
-void UNeuralNetwork::FImplBackEndUEAndORT::RunSessionAsync(const ENeuralDeviceType InDeviceType, const ENeuralDeviceType InInputDeviceType,
-	const ENeuralDeviceType InOutputDeviceType)
+void UNeuralNetwork::FImplBackEndUEAndORT::RunSessionAsync(const ENeuralDeviceType InDeviceType, const ENeuralDeviceType InInputDeviceType)
 {
 	const FScopeLock ResourcesLock(&ResoucesCriticalSection);
 
-	RunSessionImpl(InDeviceType, InInputDeviceType, InOutputDeviceType);
+	RunSessionImpl(InDeviceType, InInputDeviceType);
 
 	// Execute OnAsyncRunCompletedDelegate on main thread
 	if (DelegateThreadMode == ENeuralThreadMode::GameThread)
@@ -966,10 +944,9 @@ void UNeuralNetwork::FImplBackEndUEAndORT::RunSessionAsync(const ENeuralDeviceTy
 	}
 }
 
-void UNeuralNetwork::FImplBackEndUEAndORT::RunSessionSync(const ENeuralDeviceType InDeviceType, const ENeuralDeviceType InInputDeviceType,
-	const ENeuralDeviceType InOutputDeviceType)
+void UNeuralNetwork::FImplBackEndUEAndORT::RunSessionSync(const ENeuralDeviceType InDeviceType, const ENeuralDeviceType InInputDeviceType)
 {
-	RunSessionImpl(InDeviceType, InInputDeviceType, InOutputDeviceType);
+	RunSessionImpl(InDeviceType, InInputDeviceType);
 }
 
 // Used when uploading tensors to GPU
@@ -979,8 +956,7 @@ RDG_BUFFER_ACCESS(Upload, ERHIAccess::CopySrc)
 RDG_BUFFER_ACCESS(Input, ERHIAccess::CopyDest)
 END_SHADER_PARAMETER_STRUCT()
 
-void UNeuralNetwork::FImplBackEndUEAndORT::RunSessionImpl(const ENeuralDeviceType InDeviceType, const ENeuralDeviceType InInputDeviceType,
-	const ENeuralDeviceType InOutputDeviceType)
+void UNeuralNetwork::FImplBackEndUEAndORT::RunSessionImpl(const ENeuralDeviceType InDeviceType, const ENeuralDeviceType InInputDeviceType)
 {
 	if (InDeviceType == ENeuralDeviceType::GPU)
 	{
@@ -989,9 +965,7 @@ void UNeuralNetwork::FImplBackEndUEAndORT::RunSessionImpl(const ENeuralDeviceTyp
 
 		for (FNeuralTensor& InputTensor : InputTensors)
 		{
-// @todo: FIX!
-// check(InputTensor.GetTensorType() == ENeuralTensorType::Input, TEXT("All InputTensors should be ENeuralTensorType::Input, something is wrong"));
-			if (InputTensor.GetTensorType() == ENeuralTensorType::Input)
+			if (InInputDeviceType == ENeuralDeviceType::GPU)
 			{
 				bNeedsGPUCopy = true;
 
diff --git a/Engine/Plugins/Experimental/NeuralNetworkInference/Source/NeuralNetworkInference/Private/NeuralNetworkImplBackEndUEAndORT.h b/Engine/Plugins/Experimental/NeuralNetworkInference/Source/NeuralNetworkInference/Private/NeuralNetworkImplBackEndUEAndORT.h
--- a/Engine/Plugins/Experimental/NeuralNetworkInference/Source/NeuralNetworkInference/Private/NeuralNetworkImplBackEndUEAndORT.h
+++ b/Engine/Plugins/Experimental/NeuralNetworkInference/Source/NeuralNetworkInference/Private/NeuralNetworkImplBackEndUEAndORT.h
@@ -56,8 +56,7 @@ public:
 		const TArray<uint8>& InModelReadFromFileInBytes, const FString& InModelFullFilePath, const ENeuralDeviceType InDeviceType, const ENeuralDeviceType InInputDeviceType,
 		const ENeuralDeviceType InOutputDeviceType);
 
-	void Run(const ENeuralSynchronousMode InSynchronousMode, const ENeuralDeviceType InDeviceType, const ENeuralDeviceType InInputDeviceType,
-		const ENeuralDeviceType InOutputDeviceType);
+	void Run(const ENeuralSynchronousMode InSynchronousMode, const ENeuralDeviceType InDeviceType, const ENeuralDeviceType InInputDeviceType);
 
 #ifdef WITH_UE_AND_ORT_SUPPORT
 private:
@@ -89,7 +88,7 @@ private:
 	public:
 		FNeuralNetworkAsyncTask(UNeuralNetwork::FImplBackEndUEAndORT* InBackEnd);
 		void SetRunSessionArgs(const ENeuralSynchronousMode InSyncMode, const ENeuralDeviceType InDeviceType,
-			const ENeuralDeviceType InInputDeviceType, const ENeuralDeviceType InOutputDeviceType);
+			const ENeuralDeviceType InInputDeviceType);
 	protected:
 		void DoWork();
 		// This next section of code needs to be here. Not important as to why.
@@ -103,7 +102,6 @@ private:
 		ENeuralSynchronousMode SyncMode;
 		ENeuralDeviceType DeviceType;
 		ENeuralDeviceType InputDeviceType;
-		ENeuralDeviceType OutputDeviceType;
 	};
 	TUniquePtr<FAsyncTask<FNeuralNetworkAsyncTask>> NeuralNetworkAsyncTask;
 
@@ -118,7 +116,8 @@ private:
 		const ENeuralDeviceType InInputDeviceType, const ENeuralDeviceType InOutputDeviceType);
 
 	bool SetTensorsFromNetwork(TArray<FNeuralTensor>& OutTensors, TArray<const char*>& InTensorNames, TArray<ENeuralDataType>& InTensorDataTypes,
-		TArray<TArray<int64>>& InSizes, TArray<ENeuralTensorType>& InTensorGPUTypes, const bool bIsInput);
+		TArray<TArray<int64>>& InSizes, TArray<ENeuralTensorType>& InTensorTypes, const bool bIsInput,
+		const ENeuralDeviceType InInputDeviceType, const ENeuralDeviceType InOutputDeviceType);
 
 	static void LinkTensorToONNXRuntime(TArray<FNeuralTensor>& InOutTensors, TArray<Ort::Value>& InOutOrtTensors, Ort::MemoryInfo& InOutAllocatorInfo,
 		const int32 InTensorIndex);
@@ -127,9 +126,9 @@ private:
 	bool LinkTensorResourceToONNXRuntime(FNeuralTensor& InOutTensor, Ort::Value& InOutOrtTensor, void* InOutD3DResource);
 #endif
 
-	void RunSessionAsync(const ENeuralDeviceType InDeviceType, const ENeuralDeviceType InInputDeviceType, const ENeuralDeviceType InOutputDeviceType);
-	void RunSessionSync(const ENeuralDeviceType InDeviceType, const ENeuralDeviceType InInputDeviceType, const ENeuralDeviceType InOutputDeviceType);
-	void RunSessionImpl(const ENeuralDeviceType InDeviceType, const ENeuralDeviceType InInputDeviceType, const ENeuralDeviceType InOutputDeviceType);
+	void RunSessionAsync(const ENeuralDeviceType InDeviceType, const ENeuralDeviceType InInputDeviceType);
+	void RunSessionSync(const ENeuralDeviceType InDeviceType, const ENeuralDeviceType InInputDeviceType);
+	void RunSessionImpl(const ENeuralDeviceType InDeviceType, const ENeuralDeviceType InInputDeviceType);
 
 	void ClearResources();
 
