Adding NNI documentation:
- NNI/Public/NeuraTensor.h more cleaned and documented

_fyi Francisco.Vicente
#rnx

diff --git a/Engine/Plugins/Experimental/NeuralNetworkInference/Source/NeuralNetworkInference/Public/NeuralTensor.h b/Engine/Plugins/Experimental/NeuralNetworkInference/Source/NeuralNetworkInference/Public/NeuralTensor.h
--- a/Engine/Plugins/Experimental/NeuralNetworkInference/Source/NeuralNetworkInference/Public/NeuralTensor.h
+++ b/Engine/Plugins/Experimental/NeuralNetworkInference/Source/NeuralNetworkInference/Public/NeuralTensor.h
@@ -82,60 +82,5 @@
 	FORCEINLINE bool operator!=(const FNeuralTensor& InTensor) const;
 
 	/**
-	 * @see At(), GetArrayCopy(), GetUnderlyingUInt8ArrayRef(), GetData(), GetDataCasted(), which represent all the functions that can access/modify
-	 *  the CPU memory. These functions could result in undefined behavior in not used properly:
-	 * - If FNeuralTensor goes out of scope, it is destructed, or its memory is reset (constructor, SetNum(), etc.), these reference/pointer returned
-	 *   by any of these functions will no longer be valid (except GetUnderlyingUInt8ArrayRef()).
-	 * - These functions only modify the CPU memory, they do not synchronize CPU and GPU memory automatically. The user must do this accordingly.
-	 *
-	 * At() returns a reference to the desired indexed element of the underlying CPU data (UnderlyingUInt8ArrayData). There are 2 variations of this
-	 * function, one constant (it does not allow to modify the referenced value) and a non-const one (which allows modifying it).
-	 *
-	 * @return Reference to the indexed element (const or non-constant).
-	 */
-	template <typename T, typename TInput>
-	T& At(const TInput InIndex);
-	template <typename T, typename TInput>
-	const T& At(const TInput InIndex) const;
-
-	/**
-	 * @see At(), GetArrayCopy(), GetUnderlyingUInt8ArrayRef(), GetData(), GetDataCasted(), which represent all the functions that can access/modify
-	 *  the CPU memory. These functions could result in undefined behavior in not used properly:
-	 * - If FNeuralTensor goes out of scope, it is destructed, or its memory is reset (constructor, SetNum(), etc.), these reference/pointer returned
-	 *   by any of these functions will no longer be valid (except GetUnderlyingUInt8ArrayRef()).
-	 * - These functions only modify the CPU memory, they do not synchronize CPU and GPU memory automatically. The user must do this accordingly.
-	 *
-	 * GetArrayCopy and GetUnderlyingUInt8ArrayRef allow accessing the tensor data as a TArray, but they cannot modify the data values:
-	 * - GetArrayCopy<T> (slower but safer) returns a copy of the data as a TArray<T>. T has to be the same size than sizeof(DataType).
-	 * - GetUnderlyingUInt8ArrayRef (faster but could go out of scope) returns a reference to the underlying TArray<uint8> that contains the results.
-	 */
-	template <typename T>
-	TArray<T> GetArrayCopy() const;
-	FORCEINLINE const TArray<uint8>& GetUnderlyingUInt8ArrayRef() const;
-
-	/**
-	 * @see At(), GetArrayCopy(), GetUnderlyingUInt8ArrayRef(), GetData(), GetDataCasted(), which represent all the functions that can access/modify
-	 *  the CPU memory. These functions could result in undefined behavior in not used properly:
-	 * - If FNeuralTensor goes out of scope, it is destructed, or its memory is reset (constructor, SetNum(), etc.), these reference/pointer returned
-	 *   by any of these functions will no longer be valid (except GetUnderlyingUInt8ArrayRef()).
-	 * - These functions only modify the CPU memory, they do not synchronize CPU and GPU memory automatically. The user must do this accordingly.
-	 *
-	 * GetArrayCopy and GetUnderlyingUInt8ArrayRef allow accessing and optionally modifying the tensor data as a pointer:
-	 * - GetData() returns a void pointer.
-	 * - GetDataCasted<T>() returns a pointer casted to the desired type T.
-	 * There are 2 variations of each of these 2 functions, a constant (for read-only access) and a non-const one (which allows modifying the data).
-	 *
-	 * @return Void (GetData) or T (GetDataCasted) pointer to the underlying FNeuralTensor TArray<uint8> data (const or non-constant).
-	 */
-	FORCEINLINE void* GetData();
-	FORCEINLINE const void* const GetData() const;
-
-	template<typename T>
-	FORCEINLINE T* GetDataCasted();
-
-	template<typename T>
-	FORCEINLINE const T* const GetDataCasted() const;
-
-	/**
 	 * It returns a constant reference of DataType. @see DataType for more details.
 	 */
@@ -194,39 +139,59 @@
 	FORCEINLINE void SetEnableGPU(const bool bInEnableGPU);
 
 	/**
-	 * There are different functions to upload/download memory from/to CPU to/from GPU:
-	 * - CPUToRDGBuilder_RenderThread either uploads the memory from the CPU to the GPU (if bInShouldCopyFromCPU) or simply allocates the desired GPU
-	 *   memory. Then, it registers that GPU memory on the input FRDGBuilder.
-	 * - GPUToRDGBuilder_RenderThread implies that the memory is already on the GPU, and simply registers that GPU memory on the input FRDGBuilder.
-	 * - RDGBuilderToCPU_RenderThread downloads the memory from the GPU to the CPU. It is the opposite operation than CPUToRDGBuilder_RenderThread.
+	 * @see At(), GetArrayCopy(), GetUnderlyingUInt8ArrayRef(), GetData(), GetDataCasted(), which represent all the functions that can access/modify
+	 *  the CPU memory. These functions could result in undefined behavior in not used properly:
+	 * - If FNeuralTensor goes out of scope, it is destructed, or its memory is reset (constructor, SetNum(), etc.), these reference/pointer returned
+	 *   by any of these functions will no longer be valid (except GetUnderlyingUInt8ArrayRef()).
+	 * - These functions only modify the CPU memory, they do not synchronize CPU and GPU memory automatically. The user must do this accordingly.
 	 *
-	 * @see CPUToRDGBuilder_RenderThread, GPUToRDGBuilder_RenderThread, RDGBuilderToCPU_RenderThread
+	 * At() returns a reference to the desired indexed element of the underlying CPU data (UnderlyingUInt8ArrayData). There are 2 variations of this
+	 * function, one constant (it does not allow to modify the referenced value) and a non-const one (which allows modifying it).
+	 *
+	 * @return Reference to the indexed element (const or non-constant).
 	 */
-	void CPUToRDGBuilder_RenderThread(FRDGBuilder* InOutGraphBuilder);
-	void GPUToRDGBuilder_RenderThread(FRDGBuilder* InOutGraphBuilder);
-	void RDGBuilderToCPU_RenderThread(FRDGBuilder* InOutGraphBuilder);
+	template <typename T, typename TInput>
+	T& At(const TInput InIndex);
+	template <typename T, typename TInput>
+	const T& At(const TInput InIndex) const;
 
 	/**
-	 * It allocates the GPU data, in particular the pooled buffer.
-	 * @param InOutNativeResource Pointer to the ID3D12Resource obtained from PooledBuffer. This resource can be shared with the DirectML execution
-	 * provider.
+	 * @see At(), GetArrayCopy(), GetUnderlyingUInt8ArrayRef(), GetData(), GetDataCasted(), which represent all the functions that can access/modify
+	 *  the CPU memory. These functions could result in undefined behavior in not used properly:
+	 * - If FNeuralTensor goes out of scope, it is destructed, or its memory is reset (constructor, SetNum(), etc.), these reference/pointer returned
+	 *   by any of these functions will no longer be valid (except GetUnderlyingUInt8ArrayRef()).
+	 * - These functions only modify the CPU memory, they do not synchronize CPU and GPU memory automatically. The user must do this accordingly.
+	 *
+	 * GetArrayCopy and GetUnderlyingUInt8ArrayRef allow accessing the tensor data as a TArray, but they cannot modify the data values:
+	 * - GetArrayCopy<T> (slower but safer) returns a copy of the data as a TArray<T>. T has to be the same size than sizeof(DataType).
+	 * - GetUnderlyingUInt8ArrayRef (faster but could go out of scope) returns a reference to the underlying TArray<uint8> that contains the results.
 	 */
-	bool InitPooledBuffer(void** InOutNativeResource = nullptr);
+	template <typename T>
+	TArray<T> GetArrayCopy() const;
+	FORCEINLINE const TArray<uint8>& GetUnderlyingUInt8ArrayRef() const;
 
 	/**
-	 * It returns a constant reference of PooledBuffer. @see PooledBuffer for more details.
+	 * @see At(), GetArrayCopy(), GetUnderlyingUInt8ArrayRef(), GetData(), GetDataCasted(), which represent all the functions that can access/modify
+	 *  the CPU memory. These functions could result in undefined behavior in not used properly:
+	 * - If FNeuralTensor goes out of scope, it is destructed, or its memory is reset (constructor, SetNum(), etc.), these reference/pointer returned
+	 *   by any of these functions will no longer be valid (except GetUnderlyingUInt8ArrayRef()).
+	 * - These functions only modify the CPU memory, they do not synchronize CPU and GPU memory automatically. The user must do this accordingly.
+	 *
+	 * GetArrayCopy and GetUnderlyingUInt8ArrayRef allow accessing and optionally modifying the tensor data as a pointer:
+	 * - GetData() returns a void pointer.
+	 * - GetDataCasted<T>() returns a pointer casted to the desired type T.
+	 * There are 2 variations of each of these 2 functions, a constant (for read-only access) and a non-const one (which allows modifying the data).
+	 *
+	 * @return Void (GetData) or T (GetDataCasted) pointer to the underlying FNeuralTensor TArray<uint8> data (const or non-constant).
 	 */
-	const TRefCountPtr<FRDGPooledBuffer>& GetPooledBuffer() const;
+	FORCEINLINE void* GetData();
+	FORCEINLINE const void* const GetData() const;
+
+	template<typename T>
+	FORCEINLINE T* GetDataCasted();
+
+	template<typename T>
+	FORCEINLINE const T* const GetDataCasted() const;
 
 	/**
-	 * It returns a const reference of BufferSRVRef. @see BufferSRVRef for more details.
-	 */
-	const FRDGBufferSRVRef GetBufferSRVRef() const;
-
-	/**
-	 * It returns a reference of BufferUAVRef. @see BufferUAVRef for more details.
-	 */
-	FRDGBufferUAVRef GetBufferUAVRef();
-
-	/**
 	 * It resizes the tensor to the desired new size (i.e., modiying DataType, Sizes, Volume, UnderlyingUInt8ArrayData, etc.).
@@ -242,7 +207,6 @@
 	void SetNumUninitialized(const ENeuralDataType InDataType, const TArray<int64>& InSizes, const bool bInAllowShrinking = true);
 	FORCEINLINE void SetNumUninitialized(const ENeuralDataType InDataType, const int64 InVolume, const bool bInAllowShrinking = true);
 
-// CONTINUE HERE
 	/**
 	 * This will replace the TArray with the input one, by deeply copying the array (safer and easier to use).
 	 * For maximum performance:
@@ -256,11 +220,5 @@
 	void SetFromArrayCopy(const TArray<T>& InArray);
 
 	/**
-	 * It will fill the current neural tensor with the input FTensorProto.
-	 * @return Whether the conversion was successful.
-	 */
-	bool SetFromTensorProto(const struct FTensorProto* const InTensorProto, const FString& InTensorName, const ENeuralTensorType InTensorType);
-
-	/**
 	 * It sets all the elements of the FNeuralTensor to InValue.
 	 * It use a double typename to avoid the mistake of SetTo(0) for a double or float (because that 0 would be an int otherwise).
@@ -269,5 +227,11 @@
 	void SetTo(const TInput InValue);
 
 	/**
+	 * It will fill the current neural tensor with the input FTensorProto.
+	 * @return Whether the conversion was successful.
+	 */
+	bool SetFromTensorProto(const struct FTensorProto* const InTensorProto, const FString& InTensorName, const ENeuralTensorType InTensorType);
+
+	/**
 	 * It flips the InDimension dimension of the tensor.
 	 * @return Whether the flip was successful.
@@ -294,5 +258,40 @@
 	bool ReshapeMove(TArray<int64>& InSizes);
 
 	/**
+	 * There are different functions to upload/download memory from/to CPU to/from GPU:
+	 * - CPUToRDGBuilder_RenderThread either uploads the memory from the CPU to the GPU (if bInShouldCopyFromCPU) or simply allocates the desired GPU
+	 *   memory. Then, it registers that GPU memory on the input FRDGBuilder.
+	 * - GPUToRDGBuilder_RenderThread implies that the memory is already on the GPU, and simply registers that GPU memory on the input FRDGBuilder.
+	 * - RDGBuilderToCPU_RenderThread downloads the memory from the GPU to the CPU. It is the opposite operation than CPUToRDGBuilder_RenderThread.
+	 *
+	 * @see CPUToRDGBuilder_RenderThread, GPUToRDGBuilder_RenderThread, RDGBuilderToCPU_RenderThread
+	 */
+	void CPUToRDGBuilder_RenderThread(FRDGBuilder* InOutGraphBuilder);
+	void GPUToRDGBuilder_RenderThread(FRDGBuilder* InOutGraphBuilder);
+	void RDGBuilderToCPU_RenderThread(FRDGBuilder* InOutGraphBuilder);
+
+	/**
+	 * It allocates the GPU data, in particular the pooled buffer.
+	 * @param InOutNativeResource Pointer to the ID3D12Resource obtained from PooledBuffer. This resource can be shared with the DirectML execution
+	 * provider.
+	 */
+	bool InitPooledBuffer(void** InOutNativeResource = nullptr);
+
+	/**
+	 * It returns a constant reference of PooledBuffer. @see PooledBuffer for more details.
+	 */
+	const TRefCountPtr<FRDGPooledBuffer>& GetPooledBuffer() const;
+
+	/**
+	 * It returns a const reference of BufferSRVRef. @see BufferSRVRef for more details.
+	 */
+	const FRDGBufferSRVRef GetBufferSRVRef() const;
+
+	/**
+	 * It returns a reference of BufferUAVRef. @see BufferUAVRef for more details.
+	 */
+	FRDGBufferUAVRef GetBufferUAVRef();
+
+	/**
 	 * It returns a FString with up to InMaxNumberElementsToDisplay elements displayed. If InMaxNumberElementsToDisplay <= 0, it displays them all.
 	 * @param bInReturnOnlyData If false (default), it will print all information.
@@ -385,6 +384,51 @@
 	return !(*this == InTensor);
 }
 
+ENeuralDataType FNeuralTensor::GetDataType() const
+{
+	return DataType;
+}
+
+const TArray<int64>& FNeuralTensor::GetSizes() const
+{
+	return Sizes;
+}
+
+int32 FNeuralTensor::GetNumberDimensions() const
+{
+	return Sizes.Num();
+}
+
+int64 FNeuralTensor::Num() const
+{
+	return Volume;
+}
+
+int64 FNeuralTensor::NumInBytes() const
+{
+	return UnderlyingUInt8ArrayData.Num();
+}
+
+bool FNeuralTensor::IsEmpty() const
+{
+	return UnderlyingUInt8ArrayData.IsEmpty();
+}
+
+const FString& FNeuralTensor::GetName() const
+{
+	return Name;
+}
+
+ENeuralTensorType FNeuralTensor::GetTensorType() const
+{
+	return TensorType;
+}
+
+void FNeuralTensor::SetEnableGPU(const bool bInEnableGPU)
+{
+	bEnableGPU = bInEnableGPU;
+}
+
 template <typename T, typename TInput>
 T& FNeuralTensor::At(const TInput InIndex)
 {
@@ -440,51 +484,6 @@
 	return (T*)UnderlyingUInt8ArrayData.GetData();
 }
 
-ENeuralDataType FNeuralTensor::GetDataType() const
-{
-	return DataType;
-}
-
-const TArray<int64>& FNeuralTensor::GetSizes() const
-{
-	return Sizes;
-}
-
-int32 FNeuralTensor::GetNumberDimensions() const
-{
-	return Sizes.Num();
-}
-
-int64 FNeuralTensor::Num() const
-{
-	return Volume;
-}
-
-int64 FNeuralTensor::NumInBytes() const
-{
-	return UnderlyingUInt8ArrayData.Num();
-}
-
-bool FNeuralTensor::IsEmpty() const
-{
-	return UnderlyingUInt8ArrayData.IsEmpty();
-}
-
-const FString& FNeuralTensor::GetName() const
-{
-	return Name;
-}
-
-ENeuralTensorType FNeuralTensor::GetTensorType() const
-{
-	return TensorType;
-}
-
-void FNeuralTensor::SetEnableGPU(const bool bInEnableGPU)
-{
-	bEnableGPU = bInEnableGPU;
-}
-
 void FNeuralTensor::SetNumUninitialized(const ENeuralDataType InDataType, const int64 InVolume, const bool bInAllowShrinking)
 {
 	SetNumUninitialized(InDataType, (InVolume > 0 ? TArray<int64>({ InVolume }) : TArray<int64>({})), bInAllowShrinking);
