Adding NNI documentation:
- NNI/Public/NeuraTensor.h fully cleaned and documented

#rnx

diff --git a/Engine/Plugins/Experimental/NeuralNetworkInference/Source/NeuralNetworkInference/Private/NeuralTensor.cpp b/Engine/Plugins/Experimental/NeuralNetworkInference/Source/NeuralNetworkInference/Private/NeuralTensor.cpp
--- a/Engine/Plugins/Experimental/NeuralNetworkInference/Source/NeuralNetworkInference/Private/NeuralTensor.cpp
+++ b/Engine/Plugins/Experimental/NeuralNetworkInference/Source/NeuralNetworkInference/Private/NeuralTensor.cpp
@@ -20,7 +20,8 @@ class FPrivateNeuralTensor
 public:
 	template <typename T>
 	static FString SanitizeFloat(const T InValue);
-	static void ArrayToSanitizedString(FString& InOutTensorString, const int64 InIndexStart, const int64 InIndexEnd, const int64 InOffset, const ENeuralDataType InDataType, const FNeuralTensor& InTensor);
+	static void ArrayToSanitizedString(FString& InOutTensorString, const int64 InIndexStart, const int64 InIndexEnd, const int64 InOffset,
+		const ENeuralDataType InDataType, const FNeuralTensor& InTensor);
 
 	static void NDTensorIndexesPlus1(TArray<int32>& InOutImageAreaIndexes, const TArray<int32>& InSizes);
 
@@ -50,11 +51,13 @@ FString FPrivateNeuralTensor::SanitizeFloat(const T InValue)
 		TensorString += FNumberToStringFunction(InTensor.At<DataType>(Offset + Index)) + TEXT(" "); \
 	}
 
-void FPrivateNeuralTensor::ArrayToSanitizedString(FString& InOutTensorString, const int64 InIndexStart, const int64 InIndexEnd, const int64 InOffset, const ENeuralDataType InDataType, const FNeuralTensor& InTensor)
+void FPrivateNeuralTensor::ArrayToSanitizedString(FString& InOutTensorString, const int64 InIndexStart, const int64 InIndexEnd, const int64 InOffset,
+	const ENeuralDataType InDataType, const FNeuralTensor& InTensor)
 {
 	if (InDataType == ENeuralDataType::Float)
 	{
-		FOR_LOOP_FLOAT_TYPE_TO_SANITIZED_STRING(InOutTensorString, InTensor, InIndexStart, InIndexEnd, InOffset, FPrivateNeuralTensor::SanitizeFloat, float);
+		FOR_LOOP_FLOAT_TYPE_TO_SANITIZED_STRING(InOutTensorString, InTensor, InIndexStart, InIndexEnd, InOffset, FPrivateNeuralTensor::SanitizeFloat,
+			float);
 	}
 	else if (InDataType == ENeuralDataType::Int32)
 	{
@@ -74,7 +77,8 @@ void FPrivateNeuralTensor::ArrayToSanitizedString(FString& InOutTensorString, co
 	}
 	else
 	{
-		UE_LOG(LogNeuralNetworkInference, Warning, TEXT("FNeuralTensor::ArrayToSanitizedString(): Unknown InDataType = %d used."), (int32)InDataType);
+		UE_LOG(LogNeuralNetworkInference, Warning, TEXT("FNeuralTensor::ArrayToSanitizedString(): Unknown InDataType = %d used."),
+			(int32)InDataType);
 	}
 }
 
@@ -122,7 +126,8 @@ ENeuralDataType FPrivateNeuralTensor::GetDataTypeFromTensorProtoDataType(const E
 	{
 		return ENeuralDataType::None;
 	}
-	UE_LOG(LogNeuralNetworkInference, Warning, TEXT("FNeuralTensor::GetDataTypeFromTensorProtoDataType(): Unknown InTensorProtoDataType = %d used."), (int32)InTensorProtoDataType);
+	UE_LOG(LogNeuralNetworkInference, Warning, TEXT("FNeuralTensor::GetDataTypeFromTensorProtoDataType(): Unknown InTensorProtoDataType = %d used."),
+		(int32)InTensorProtoDataType);
 	return ENeuralDataType::None;
 }
 
@@ -131,7 +136,8 @@ ENeuralDataType FPrivateNeuralTensor::GetDataTypeFromTensorProtoDataType(const E
 /* FNeuralTensor costructors
  *****************************************************************************/
 
-FNeuralTensor::FNeuralTensor(const ENeuralDataType InDataType, const TArray<int64>& InSizes, const FString& InName, const ENeuralTensorType InTensorType)
+FNeuralTensor::FNeuralTensor(const ENeuralDataType InDataType, const TArray<int64>& InSizes, const FString& InName,
+	const ENeuralTensorType InTensorType)
 	: DataType(ENeuralDataType::None)
 	, Name(InName)
 	, TensorType(InTensorType)
@@ -151,7 +157,8 @@ FNeuralTensor::FNeuralTensor(const ENeuralDataType InDataType, const int64 InVol
 /* FNeuralTensor private costructor
  *****************************************************************************/
 
-FNeuralTensor::FNeuralTensor(const ENeuralDataType InDataType, const void* const InValues, const int64 InSizeOfT, const int64 InValueNum, const TArray<int64>& InSizes, const FString& InName, const ENeuralTensorType InTensorType)
+FNeuralTensor::FNeuralTensor(const ENeuralDataType InDataType, const void* const InValues, const int64 InSizeOfT, const int64 InValueNum,
+	const TArray<int64>& InSizes, const FString& InName, const ENeuralTensorType InTensorType)
 	: FNeuralTensor(InDataType, InSizes, InName, InTensorType)
 {
 	// Sanity check
@@ -292,8 +299,8 @@ void FNeuralTensor::SetFromUnderlyingUInt8ArrayCopy(const TArray<uint8>& InArray
 {
 	if (NumInBytes() != InArray.Num())
 	{
-		UE_LOG(LogNeuralNetworkInference, Warning,
-			TEXT("FNeuralTensor::SetFromUnderlyingUInt8ArrayCopy(): NumInBytes() == InArray.Num() failed, %d != %d."), NumInBytes(), InArray.Num());
+		UE_LOG(LogNeuralNetworkInference, Warning, TEXT("FNeuralTensor::SetFromUnderlyingUInt8ArrayCopy():"
+			" NumInBytes() == InArray.Num() failed, %d != %d."), NumInBytes(), InArray.Num());
 		return;
 	}
 	UnderlyingUInt8ArrayData = InArray;
@@ -307,7 +314,7 @@ bool FNeuralTensor::SetFromTensorProto(const FTensorProto* const InTensorProto,
 		return false;
 	}
 
-	// const FString& InExternalDataDirectory = TEXT("") - @param InExternalDataDirectory is only required if InTensorProto->ExternalData is being used.
+	// const FString& InExternalDataDirectory = TEXT(""); // InExternalDataDirectory only required if InTensorProto->ExternalData is being used.
 
 	// Create Tensor
 	Name = InTensorName;
@@ -328,15 +335,18 @@ bool FNeuralTensor::SetFromTensorProto(const FTensorProto* const InTensorProto,
 		// // Sanity check
 		// if (InTensorProto->ExternalData[0].Key != TEXT("location") || InTensorProto->ExternalData[0].Value.Len() < 1)
 		// {
-		// 	UE_LOG(LogNeuralNetworkInference, Warning, TEXT("FNeuralTensor::SetFromTensorProto(): InTensorProto->ExternalData[0].Key = %s != \"location\""
-		// 		" || InTensorProto->ExternalData[0].Value.Len() = %d (should be > 0)."), *InTensorProto->ExternalData[0].Key, InTensorProto->ExternalData[0].Value.Len());
+		// 	UE_LOG(LogNeuralNetworkInference, Warning, TEXT("FNeuralTensor::SetFromTensorProto():"
+		//      " InTensorProto->ExternalData[0].Key = %s != \"location\" || InTensorProto->ExternalData[0].Value.Len() = %d (should be > 0)."),
+		//      *InTensorProto->ExternalData[0].Key, InTensorProto->ExternalData[0].Value.Len());
 		// 	return false;
 		// }
 		// // Read neural tensor from binary data
 		// const FString BinaryWeightFilePath = InExternalDataDirectory / InTensorProto->ExternalData[0].Value;
-		// if (!FModelProtoFileReader::ReadWeightsFromOtxtBinaryFile((char*)GetData(), Num() * FNeuralDataTypeUtils::GetByteSize(NeuralDataType), BinaryWeightFilePath))
+		// if (!FModelProtoFileReader::ReadWeightsFromOtxtBinaryFile((char*)GetData(), Num() * FNeuralDataTypeUtils::GetByteSize(NeuralDataType),
+		//     BinaryWeightFilePath))
 		// {
-		// 	UE_LOG(LogNeuralNetworkInference, Warning, TEXT("FNeuralTensor::SetFromTensorProto(): Could not read binary file: %s."), *BinaryWeightFilePath);
+		// 	UE_LOG(LogNeuralNetworkInference, Warning, TEXT("FNeuralTensor::SetFromTensorProto(): Could not read binary file: %s."),
+		//      *BinaryWeightFilePath);
 		// 	return false;
 		// }
 	}
@@ -379,7 +389,7 @@ bool FNeuralTensor::Flip(const int32 InDimension)
 		TensorNDSizes.Init(1, TensorNDSize);
 		for (int32 NDTensorSizeIndex = 0; NDTensorSizeIndex < TensorNDSize; ++NDTensorSizeIndex)
 		{
-			TensorNDSizes[NDTensorSizeIndex] *= Sizes[NDTensorSizeIndex]; // Idea: Sizes=[1, 2, 3, 4], InDimension = 2, then TensorNDSizes = [1, 2x3x4]
+			TensorNDSizes[NDTensorSizeIndex] *= Sizes[NDTensorSizeIndex]; // E.g., Sizes=[1,2,3,4] and InDimension=2 --> TensorNDSizes = [1, 2x3x4]
 		}
 	}
 
@@ -398,7 +408,8 @@ bool FNeuralTensor::Flip(const int32 InDimension)
 		// Remove last index (that makes it a normal index) and replace with its flipped equivalent
 		FlippedTensorIndex = FlippedTensorIndex + Sizes[InDimension] - 1 - 2 * TensorNDIndexes.Last();
 		// Flip TensorIndex value
-		FMemory::Memcpy(&NewArrayOnCPU.GetData()[TensorIndex * BytesPerIndex], &UnderlyingUInt8ArrayData.GetData()[FlippedTensorIndex * DimensionOffsetInBytes], DimensionOffsetInBytes);
+		FMemory::Memcpy(&NewArrayOnCPU.GetData()[TensorIndex * BytesPerIndex],
+			&UnderlyingUInt8ArrayData.GetData()[FlippedTensorIndex * DimensionOffsetInBytes], DimensionOffsetInBytes);
 		// Increase TensorNDIndexes
 		FPrivateNeuralTensor::NDTensorIndexesPlus1(TensorNDIndexes, TensorNDSizes);
 	}
@@ -411,12 +422,14 @@ bool FNeuralTensor::Flip(const int32 InDimensionFirst, const int32 InDimensionLa
 	// Sanity checks
 	if (InDimensionFirst >= InDimensionLast)
 	{
-		UE_LOG(LogNeuralNetworkInference, Warning, TEXT("FNeuralTensor-%s::Transpose(): InDimensionFirst < InDimensionLast failed, %d >= %d."), *Name, InDimensionFirst, InDimensionLast);
+		UE_LOG(LogNeuralNetworkInference, Warning, TEXT("FNeuralTensor-%s::Flip(): InDimensionFirst < InDimensionLast failed, %d >= %d."),
+			*Name, InDimensionFirst, InDimensionLast);
 		return false;
 	}
 	else if (InDimensionLast > GetNumberDimensions())
 	{
-		UE_LOG(LogNeuralNetworkInference, Warning, TEXT("FNeuralTensor-%s::Transpose(): InDimensionLast < GetNumberDimensions() failed, %d >= %d."), *Name, InDimensionLast, GetNumberDimensions());
+		UE_LOG(LogNeuralNetworkInference, Warning, TEXT("FNeuralTensor-%s::Flip(): InDimensionLast < GetNumberDimensions() failed, %d >= %d."),
+			*Name, InDimensionLast, GetNumberDimensions());
 		return false;
 	}
 	// Flip
@@ -466,7 +479,8 @@ bool FNeuralTensor::Transpose()
 		}
 		else
 		{
-			UE_LOG(LogNeuralNetworkInference, Warning, TEXT("FNeuralTensor-%s::Transpose(): Unexpected case NumberDimensions = %d != 1 || 2."), *Name, NumberDimensions);
+			UE_LOG(LogNeuralNetworkInference, Warning, TEXT("FNeuralTensor-%s::Transpose(): Unexpected case NumberDimensions = %d != 1 || 2."),
+				*Name, NumberDimensions);
 			return false;
 		}
 		// Swap W <-> H
@@ -489,7 +503,8 @@ bool FNeuralTensor::ReshapeMove(TArray<int64>& InSizes)
 		Swap(Sizes, InSizes);
 		return true;
 	}
-	UE_LOG(LogNeuralNetworkInference, Warning, TEXT("FNeuralTensor-%s::ReshapeMove(): Volume == NewVolume failed, %d != %d."), *Name, Volume, NewVolume);
+	UE_LOG(LogNeuralNetworkInference, Warning, TEXT("FNeuralTensor-%s::ReshapeMove(): Volume == NewVolume failed, %d != %d."), *Name, Volume,
+		NewVolume);
 	return false;
 }
 
@@ -501,12 +516,15 @@ void FNeuralTensor::CPUToRDGBuilder_RenderThread(FRDGBuilder* InOutGraphBuilder)
 // at D:/P4/private_dh_research_pitt/Engine/Source/Runtime/Windows/D3D11RHI/Private/D3D11VertexBuffer.cpp:109 
 // with error E_INVALIDARG
 	// Idea:
-	// - EBufferUsageFlags::Volatile: Updated multiple times in a frame, but does not imply a lifetime of 1 frame. E.g. a vertex buffer you update every frame with new vertices.
+	// - EBufferUsageFlags::Volatile: Updated multiple times in a frame, but does not imply a lifetime of 1 frame. E.g. a vertex buffer you update
+	//   every frame with new vertices.
 	// - EBufferUsageFlags::Transient: Used during 1 frame. Volatile and transient are not mutually exclusive.
 	// - EBufferUsageFlags::KeepCPUAccessible: Not needed, I can just copy the final GPU memory back to RAM at the very end
-	// - Input and Intermediate(Not)Initialized currently share the same attributes because input might become intermediate (e.g., if input tensor fed into a ReLU, which simply modifies
-	//   the input FNeuralTensor). However, Intermediate(Not)Initialized and Output do not copy the memory from CPU to GPU but rather simply allocates it.
-	// - Output might also become Intermediate(Not)Initialized (e.g., if Output -> ReLU -> Output), so it is kept as ReadWrite rather than written once to account for this.
+	// - Input and Intermediate(Not)Initialized currently share the same attributes because input might become intermediate (e.g., if input tensor
+	//   fed into a ReLU, which simply modifies the input FNeuralTensor). However, Intermediate(Not)Initialized and Output do not copy the memory
+	//   from CPU to GPU but rather simply allocates it.
+	// - Output might also become Intermediate(Not)Initialized (e.g., if Output -> ReLU -> Output), so it is kept as ReadWrite rather than written
+	//   once to account for this.
 	// Call ToGPU_RenderThread with the right flags
 	if (TensorType == ENeuralTensorType::Generic)
 	{
@@ -518,7 +536,8 @@ void FNeuralTensor::CPUToRDGBuilder_RenderThread(FRDGBuilder* InOutGraphBuilder)
 	}
 	else if (TensorType == ENeuralTensorType::IntermediateNotInitialized)
 	{
-		return CPUToRDGBuilder_RenderThread(InOutGraphBuilder, EBufferUsageFlags::ShaderResource | EBufferUsageFlags::UnorderedAccess | EBufferUsageFlags::Transient, false);
+		return CPUToRDGBuilder_RenderThread(InOutGraphBuilder,
+			EBufferUsageFlags::ShaderResource | EBufferUsageFlags::UnorderedAccess | EBufferUsageFlags::Transient, false);
 	}
 	else if (TensorType == ENeuralTensorType::IntermediateInitialized)
 	{
@@ -546,7 +565,8 @@ void FNeuralTensor::GPUToRDGBuilder_RenderThread(FRDGBuilder* InOutGraphBuilder)
 	}
 	// Sanity checks
 	checkf(IsInRenderingThread(), TEXT("FNeuralTensor-%s::GPUToRDGBuilder_RenderThread(): IsInRenderingThread() must be true."), *Name);
-	checkf(PooledBuffer.IsValid() && InOutGraphBuilder, TEXT("FNeuralTensor-%s::GPUToRDGBuilder_RenderThread(): IPooledBuffer and InOutGraphBuilder cannot be nullptrs."), *Name);
+	checkf(PooledBuffer.IsValid() && InOutGraphBuilder, TEXT("FNeuralTensor-%s::GPUToRDGBuilder_RenderThread(): IPooledBuffer and InOutGraphBuilder"
+		" cannot be nullptrs."), *Name);
 	// Register BufferRef
 	FRDGBufferRef BufferRef = InOutGraphBuilder->RegisterExternalBuffer(*PooledBuffer);
 	// Recreate BufferSRVRef
@@ -797,9 +817,9 @@ void FNeuralTensor::SetFromPointer(const void* const InData, const int64 InSizeO
 	// Sanity checks
 	if (Num() != InDataSize || NumInBytes() != InSizeOfT * InDataSize)
 	{
-		UE_LOG(LogNeuralNetworkInference, Warning, TEXT("FNeuralTensor-%s::SetFromPointer: Num() == InDataSize failed, %d vs. %d, or NumInBytes() == sizeof(T) x InDataSize failed, %d vs. %d."
-			" If you want to modify the dimensions of FNeuralTensor, call SetNumUninitialized() first."),
-			*Name, Num(), InDataSize, NumInBytes(), InSizeOfT * InDataSize);
+		UE_LOG(LogNeuralNetworkInference, Warning, TEXT("FNeuralTensor-%s::SetFromPointer: Num() == InDataSize failed, %d vs. %d, or"
+			" NumInBytes() == sizeof(T) x InDataSize failed, %d vs. %d. If you want to modify the dimensions of FNeuralTensor, call"
+			" SetNumUninitialized() first."), *Name, Num(), InDataSize, NumInBytes(), InSizeOfT * InDataSize);
 	}
 	else
 	{
@@ -808,7 +828,7 @@ void FNeuralTensor::SetFromPointer(const void* const InData, const int64 InSizeO
 	}
 }
 
-bool FNeuralTensor::CheckTAndDataTypeResult(const bool bInCheckTAndDataTypeResult, const int64 InSizeOfT) const
+bool FNeuralTensor::CheckTAndDataTypeEquivalentAuxiliary(const bool bInCheckTAndDataTypeResult, const int64 InSizeOfT) const
 {
 	const int64 ByteSizeOfDataType = FNeuralDataTypeUtils::GetByteSize(DataType);
 	if (!bInCheckTAndDataTypeResult)
@@ -817,20 +837,21 @@ bool FNeuralTensor::CheckTAndDataTypeResult(const bool bInCheckTAndDataTypeResul
 		// sizeof(T) and DataType do not match
 		if (ByteSizeOfDataType != InSizeOfT)
 		{
-			UE_LOG(LogNeuralNetworkInference, Warning, TEXT("FNeuralTensor-%s::CheckTAndDataTypeResult() failed: DataType = %s, but sizeof(%s) = %d != sizeof(T) = %d."),
-				*Name, *DataTypeString, *DataTypeString, ByteSizeOfDataType, InSizeOfT);
+			UE_LOG(LogNeuralNetworkInference, Warning, TEXT("FNeuralTensor-%s::CheckTAndDataTypeEquivalentAuxiliary() failed: DataType = %s, but"
+				" sizeof(%s) = %d != sizeof(T) = %d."), *Name, *DataTypeString, *DataTypeString, ByteSizeOfDataType, InSizeOfT);
 		}
 		// sizeof(T) matches, but not the expected DataType
 		else
 		{
-			UE_LOG(LogNeuralNetworkInference, Warning, TEXT("FNeuralTensor-%s::CheckTAndDataTypeResult() failed: DataType = %s, but used a different DataType with the same sizeof(%s) of %d."),
-				*Name, *DataTypeString, *DataTypeString, InSizeOfT);
+			UE_LOG(LogNeuralNetworkInference, Warning, TEXT("FNeuralTensor-%s::CheckTAndDataTypeEquivalentAuxiliary() failed: DataType = %s, but"
+				" used a different DataType with the same sizeof(%s) of %d."), *Name, *DataTypeString, *DataTypeString, InSizeOfT);
 		}
 	}
 	return bInCheckTAndDataTypeResult;
 }
 
-void FNeuralTensor::CPUToRDGBuilder_RenderThread(FRDGBuilder* InOutGraphBuilder, const EBufferUsageFlags InBufferUsageFlags, const bool bInShouldCopyFromCPU)
+void FNeuralTensor::CPUToRDGBuilder_RenderThread(FRDGBuilder* InOutGraphBuilder, const EBufferUsageFlags InBufferUsageFlags,
+	const bool bInShouldCopyFromCPU)
 {
 	// Sanity checks
 	if (!bEnableGPU || IsEmpty())
@@ -861,7 +882,8 @@ void FNeuralTensor::CPUToRDGBuilder_RenderThread(FRDGBuilder* InOutGraphBuilder,
 		BufferDesc.UnderlyingType = FRDGBufferDesc::EUnderlyingType::VertexBuffer;
 		
 		BufferRef = bInShouldCopyFromCPU 
-			? CreateVertexBuffer(*InOutGraphBuilder, *Name, BufferDesc, UnderlyingUInt8ArrayData.GetData(), NumInBytes(), ERDGInitialDataFlags::NoCopy)
+			? CreateVertexBuffer(*InOutGraphBuilder, *Name, BufferDesc, UnderlyingUInt8ArrayData.GetData(), NumInBytes(),
+				ERDGInitialDataFlags::NoCopy)
 			: InOutGraphBuilder->CreateBuffer(BufferDesc, *Name);
 
 		// Recreate PooledBuffer for future runs
diff --git a/Engine/Plugins/Experimental/NeuralNetworkInference/Source/NeuralNetworkInference/Public/NeuralTensor.h b/Engine/Plugins/Experimental/NeuralNetworkInference/Source/NeuralNetworkInference/Public/NeuralTensor.h
--- a/Engine/Plugins/Experimental/NeuralNetworkInference/Source/NeuralNetworkInference/Public/NeuralTensor.h
+++ b/Engine/Plugins/Experimental/NeuralNetworkInference/Source/NeuralNetworkInference/Public/NeuralTensor.h
@@ -139,10 +139,11 @@ public:
 	FORCEINLINE void SetEnableGPU(const bool bInEnableGPU);
 
 	/**
-	 * @see At(), GetArrayCopy(), GetUnderlyingUInt8ArrayRef(), GetData(), GetDataCasted(), which represent all the functions that can access/modify
-	 *  the CPU memory. These functions could result in undefined behavior in not used properly:
-	 * - If FNeuralTensor goes out of scope, it is destructed, or its memory is reset (constructor, SetNum(), etc.), these reference/pointer returned
-	 *   by any of these functions will no longer be valid (except GetUnderlyingUInt8ArrayRef()).
+	 * @see At(), GetArrayCopy(), GetUnderlyingUInt8ArrayRef(), GetData(), GetDataCasted(), SetNumUninitialized, SetFromUnderlyingUInt8ArrayCopy,
+	 * SetFromArrayCopy, SetTo, SetFromTensorProto, which represent all the functions that can access/modify the CPU memory. Some of these functions
+	 * could result in undefined behavior in not used properly:
+	 * - If FNeuralTensor goes out of scope, it is destructed, or its memory is reset (constructor, SetNum(), etc.), the functions returning
+	 *   references/pointers will no longer be valid.
 	 * - These functions only modify the CPU memory, they do not synchronize CPU and GPU memory automatically. The user must do this accordingly.
 	 *
 	 * At() returns a reference to the desired indexed element of the underlying CPU data (UnderlyingUInt8ArrayData). There are 2 variations of this
@@ -156,10 +157,11 @@ public:
 	const T& At(const TInput InIndex) const;
 
 	/**
-	 * @see At(), GetArrayCopy(), GetUnderlyingUInt8ArrayRef(), GetData(), GetDataCasted(), which represent all the functions that can access/modify
-	 *  the CPU memory. These functions could result in undefined behavior in not used properly:
-	 * - If FNeuralTensor goes out of scope, it is destructed, or its memory is reset (constructor, SetNum(), etc.), these reference/pointer returned
-	 *   by any of these functions will no longer be valid (except GetUnderlyingUInt8ArrayRef()).
+	 * @see At(), GetArrayCopy(), GetUnderlyingUInt8ArrayRef(), GetData(), GetDataCasted(), SetNumUninitialized, SetFromUnderlyingUInt8ArrayCopy,
+	 * SetFromArrayCopy, SetTo, SetFromTensorProto, which represent all the functions that can access/modify the CPU memory. Some of these functions
+	 * could result in undefined behavior in not used properly:
+	 * - If FNeuralTensor goes out of scope, it is destructed, or its memory is reset (constructor, SetNum(), etc.), the functions returning
+	 *   references/pointers will no longer be valid.
 	 * - These functions only modify the CPU memory, they do not synchronize CPU and GPU memory automatically. The user must do this accordingly.
 	 *
 	 * GetArrayCopy and GetUnderlyingUInt8ArrayRef allow accessing the tensor data as a TArray, but they cannot modify the data values:
@@ -171,10 +173,11 @@ public:
 	FORCEINLINE const TArray<uint8>& GetUnderlyingUInt8ArrayRef() const;
 
 	/**
-	 * @see At(), GetArrayCopy(), GetUnderlyingUInt8ArrayRef(), GetData(), GetDataCasted(), which represent all the functions that can access/modify
-	 *  the CPU memory. These functions could result in undefined behavior in not used properly:
-	 * - If FNeuralTensor goes out of scope, it is destructed, or its memory is reset (constructor, SetNum(), etc.), these reference/pointer returned
-	 *   by any of these functions will no longer be valid (except GetUnderlyingUInt8ArrayRef()).
+	 * @see At(), GetArrayCopy(), GetUnderlyingUInt8ArrayRef(), GetData(), GetDataCasted(), SetNumUninitialized, SetFromUnderlyingUInt8ArrayCopy,
+	 * SetFromArrayCopy, SetTo, SetFromTensorProto, which represent all the functions that can access/modify the CPU memory. Some of these functions
+	 * could result in undefined behavior in not used properly:
+	 * - If FNeuralTensor goes out of scope, it is destructed, or its memory is reset (constructor, SetNum(), etc.), the functions returning
+	 *   references/pointers will no longer be valid.
 	 * - These functions only modify the CPU memory, they do not synchronize CPU and GPU memory automatically. The user must do this accordingly.
 	 *
 	 * GetArrayCopy and GetUnderlyingUInt8ArrayRef allow accessing and optionally modifying the tensor data as a pointer:
@@ -194,7 +197,14 @@ public:
 	FORCEINLINE const T* const GetDataCasted() const;
 
 	/**
-	 * It resizes the tensor to the desired new size (i.e., modiying DataType, Sizes, Volume, UnderlyingUInt8ArrayData, etc.).
+	 * @see At(), GetArrayCopy(), GetUnderlyingUInt8ArrayRef(), GetData(), GetDataCasted(), SetNumUninitialized, SetFromUnderlyingUInt8ArrayCopy,
+	 * SetFromArrayCopy, SetTo, SetFromTensorProto, which represent all the functions that can access/modify the CPU memory. Some of these functions
+	 * could result in undefined behavior in not used properly:
+	 * - If FNeuralTensor goes out of scope, it is destructed, or its memory is reset (constructor, SetNum(), etc.), the functions returning
+	 *   references/pointers will no longer be valid.
+	 * - These functions only modify the CPU memory, they do not synchronize CPU and GPU memory automatically. The user must do this accordingly.
+	 *
+	 * SetNumUninitialized resizes the tensor to the desired new size (i.e., modiying DataType, Sizes, Volume, UnderlyingUInt8ArrayData, etc.).
 	 * @param InDataType It sets DataType. If None, it will keep its current value. @see DataType for more details.
 	 * @param InSizes It sets Sizes and Volume. Set to empty (or omit this argument) if memory allocation is not required or the final tensor size is
 	 *  unknown. @see Sizes, Volume for more details.
@@ -208,51 +218,78 @@ public:
 	FORCEINLINE void SetNumUninitialized(const ENeuralDataType InDataType, const int64 InVolume, const bool bInAllowShrinking = true);
 
 	/**
-	 * This will replace the TArray with the input one, by deeply copying the array (safer and easier to use).
-	 * For maximum performance:
+	 * @see At(), GetArrayCopy(), GetUnderlyingUInt8ArrayRef(), GetData(), GetDataCasted(), SetNumUninitialized, SetFromUnderlyingUInt8ArrayCopy,
+	 * SetFromArrayCopy, SetTo, SetFromTensorProto, which represent all the functions that can access/modify the CPU memory. Some of these functions
+	 * could result in undefined behavior in not used properly:
+	 * - If FNeuralTensor goes out of scope, it is destructed, or its memory is reset (constructor, SetNum(), etc.), the functions returning
+	 *   references/pointers will no longer be valid.
+	 * - These functions only modify the CPU memory, they do not synchronize CPU and GPU memory automatically. The user must do this accordingly.
+	 *
+	 * SetFromUnderlyingUInt8ArrayCopy/SetFromArrayCopy will replace the contents of the underlying CPU TArray with the input one, by deeply copying
+	 * the array (safer and easier to use). For maximum performance:
 	 * - If you are given an input TArray (or a FNeuralTensor with different Sizes) that cannot be moved, use SetFromArrayCopy().
-	 * - If you are initializing the FNeuralTensor iteratively (and to avoid a copy), fill the tensor with GetDataCasted().
-	 * Modifying the array InArray after calling this function will not modify the internal FNeuralTensor array.
-	 * The size of both (InArray and current FNeuralTensor) must match, i.e., Num() must match InArray.Num().
+	 * - If you are initializing the FNeuralTensor iteratively (and to avoid a copy), fill the tensor with GetData()/GetDataCasted<T>().
+	 * Modifying InArray after calling this function will not modify the data of this FNeuralTensor.
+	 * The size of InArray and this tensor must match, i.e., NumInBytes() must match TArray<uint8>::Num() or Num() match TArray<T>.Num().
 	 */
 	void SetFromUnderlyingUInt8ArrayCopy(const TArray<uint8>& InArray);
 	template<typename T>
 	void SetFromArrayCopy(const TArray<T>& InArray);
 
 	/**
-	 * It sets all the elements of the FNeuralTensor to InValue.
-	 * It use a double typename to avoid the mistake of SetTo(0) for a double or float (because that 0 would be an int otherwise).
+	 * @see At(), GetArrayCopy(), GetUnderlyingUInt8ArrayRef(), GetData(), GetDataCasted(), SetNumUninitialized, SetFromUnderlyingUInt8ArrayCopy,
+	 * SetFromArrayCopy, SetTo, SetFromTensorProto, which represent all the functions that can access/modify the CPU memory. Some of these functions
+	 * could result in undefined behavior in not used properly:
+	 * - If FNeuralTensor goes out of scope, it is destructed, or its memory is reset (constructor, SetNum(), etc.), the functions returning
+	 *   references/pointers will no longer be valid.
+	 * - These functions only modify the CPU memory, they do not synchronize CPU and GPU memory automatically. The user must do this accordingly.
+	 *
+	 * SetTo() sets all the elements of the FNeuralTensor to the same value (InValue).
+	 * It uses a double typename to avoid the mistake of SetTo(0) for a non-32-byte type (because that 0 would be an int32 otherwise).
 	 */
 	template<typename T, typename TInput>
 	void SetTo(const TInput InValue);
 
 	/**
-	 * It will fill the current neural tensor with the input FTensorProto.
+	 * @see At(), GetArrayCopy(), GetUnderlyingUInt8ArrayRef(), GetData(), GetDataCasted(), SetNumUninitialized, SetFromUnderlyingUInt8ArrayCopy,
+	 * SetFromArrayCopy, SetTo, SetFromTensorProto, which represent all the functions that can access/modify the CPU memory. Some of these functions
+	 * could result in undefined behavior in not used properly:
+	 * - If FNeuralTensor goes out of scope, it is destructed, or its memory is reset (constructor, SetNum(), etc.), the functions returning
+	 *   references/pointers will no longer be valid.
+	 * - These functions only modify the CPU memory, they do not synchronize CPU and GPU memory automatically. The user must do this accordingly.
+	 *
+	 * SetFromTensorProto() fills the current neural tensor with the input FTensorProto.
 	 * @return Whether the conversion was successful.
 	 */
 	bool SetFromTensorProto(const struct FTensorProto* const InTensorProto, const FString& InTensorName, const ENeuralTensorType InTensorType);
 
 	/**
-	 * It flips the InDimension dimension of the tensor.
-	 * @return Whether the flip was successful.
+	 * It reverses/flips the order of elements in an array along the given axis. It is needed for applications such as convolution.
+	 * - If 1 argument, it flips the InDimension dimension of the tensor.
+	 * - If 2 arguments, it flips all the dimensions of the tensor in the range [InDimensionFirst, InDimensionLast).
+	 * E.g., for the following [2x2x2] tensor:
+	 *   - Tensor    = [[[0, 1], [2, 3]], [[4, 5], [6, 7]]]
+	 *   - Flip(0)   = [[[4, 5], [6, 7]], [0, 1], [2, 3]]]
+	 *   - Flip(1)   = [[[2, 3], [0, 1]], [[6, 7], [4, 5]]]
+	 *   - Flip(0,3) = [[[7, 6], [5, 4]], [[3, 2], [1, 0]]]
+	 * @return Whether the flip was successful. It could only fail if the given arguments are not valid.
+	 * @see https://numpy.org/doc/stable/reference/generated/numpy.flip.html for more information.
 	 */
 	bool Flip(const int32 InDimension);
-
-	/**
-	 * It flips all the dimensions of the tensor in the range [InDimensionFirst, InDimensionLast). Needed for efficient convolution.
-	 * @return Whether the flip was successful.
-	 */
 	bool Flip(const int32 InDimensionFirst, const int32 InDimensionLast);
 
 	/**
-	 * It transposes the matrix (if the tensor has up to 2 dimensions).
-	 * @return Whether the transpose was successful.
+	 * It transposes the matrix if the tensor has up to 2 dimensions. Otherwise, it will return false and do nothing.
+	 * @return Whether the transpose was successful. E.g., it will fail if GetNumberDimensions() > 2.
+	 * @see https://numpy.org/doc/stable/reference/generated/numpy.transpose.html for more information.
 	 */
 	bool Transpose();
 
 	/**
-	 * If Num() is constant, it reshapes the current tensor. I.e., it just updates Sizes with InSizes.
+	 * If the number of elements is preserved (i.e., Prod(Sizes) == Prod(InSizes)), it reshapes the current tensor, meaning it updates Sizes with
+	 * InSizes. No other internal variable is modified.
 	 * Reshape() copies and ReshapeMove() moves InSizes.
+	 * @return Whether the reshape was successful. It would fail if and only if Prod(Sizes) != Prod(InSizes).
 	 */
 	bool Reshape(const TArray<int64>& InSizes);
 	bool ReshapeMove(TArray<int64>& InSizes);
@@ -293,10 +330,11 @@ public:
 	FRDGBufferUAVRef GetBufferUAVRef();
 
 	/**
-	 * It returns a FString with up to InMaxNumberElementsToDisplay elements displayed. If InMaxNumberElementsToDisplay <= 0, it displays them all.
-	 * @param bInReturnOnlyData If false (default), it will print all information.
-	 * E.g., "FNeuralTensor: Int64, Generic, volume=3, sizes={3}, data=[1 2 3]".
-	 * If false, it will simply print the data. I.e., "[1 2 3]" for the previous example.
+	 * It returns a FString with all the tensor information (underlying CPU data and optionally DataType, Sizes, Volume, etc.).
+	 * @param InMaxNumberElementsToDisplay The maximum number of elements to display. If Volume <= InMaxNumberElementsToDisplay or
+	 * InMaxNumberElementsToDisplay < 1, all elements are displayed. Otherwise, only the first and last InMaxNumberElementsToDisplay/2 are displayed.
+	 * @param bInReturnOnlyData If false (default), it will print other information such as DataType, Sizes, Volume, etc. If true, it will only print
+	 * the underlying CPU data values. E.g., "FNeuralTensor: Int64, Generic, volume=3, sizes={3}, data=[1 2 3]" vs. "[1 2 3]", respectively.
 	 */
 	FString ToString(const int64 InMaxNumberElementsToDisplay = -1, const bool bInReturnOnlyData = false) const;
 
@@ -324,15 +362,18 @@ protected:
 
 private:
 	/**
-	 * The array containing the CPU data. It is pre-allocated based on the combination of DataType and Sizes/Volume (if Sizes/Volume are not
-	 * empty/zero): UnderlyingUInt8ArrayData.Num() == sizeof(data type) x Volume.
+	 * The array containing the underlying CPU data. It is pre-allocated based on the combination of DataType and Sizes/Volume (if Sizes/Volume are
+	 * not empty/zero): UnderlyingUInt8ArrayData.Num() == sizeof(data type) x Volume.
 	 */
 	UPROPERTY()
 	TArray<uint8> UnderlyingUInt8ArrayData;
 	/**
-	 * GPU-based variables that are transient (do not need to be serialized) and initialized on the fly (with ToGPU/ToCPU).
-	 * - bEnableGPU: By default false, meaning all GPU memory will be disabled and those functions will not do anything. Enable to allow using the
-	 *   GPU functions and variables of the tensor.
+	 * GPU-related variables that are transient. They are not serialized and are re-configured every time a tensor is loaded.
+	 * - bEnableGPU: By default false, meaning all GPU memory allocation will be disabled and GPU functions will do nothing. Enable to allow using
+	 *   the GPU capabilities of the tensor.
+	 * - PooledBuffer: FRDGPooledBuffer containing the GPU memory of this tensor.
+	 * - BufferSRVRef: Casting of the PooledBuffer into a read-only SRV.
+	 * - BufferUAVRef: Casting of the PooledBuffer into a read-write UAV.
 	 */
 	UPROPERTY(Transient, VisibleAnywhere, Category = "Neural Network Inference")
 	bool bEnableGPU;
@@ -341,27 +382,31 @@ private:
 	TSharedPtr<FRDGBufferUAVRef> BufferUAVRef;
 
 	/**
-	 * Auxiliary for the templated constructor FNeuralTensor(const ENeuralDataType InDataType, const TArray<T>& InArray, ...).
+	 * Auxiliary constructor for the templated constructor that takes a TArray<T>& InArray.
 	 * InSizeOfT and sizeof(InDataType) should match, as well as the volume of InValueNum and InSizes.
+	 * It simply moves as much code as possible from the template constructor to the cpp file.
 	 */
 	FNeuralTensor(const ENeuralDataType InDataType, const void* const InValues, const int64 InSizeOfT, const int64 InValueNum,
 		const TArray<int64>& InSizes, const FString& InName, const ENeuralTensorType InTensorType);
 
 	/**
-	 * Auxiliary for the templated SetFromArrayCopy().
+	 * Auxiliary function for the templated SetFromArrayCopy().
+	 * It simply moves as much code as possible from the template function to the cpp file.
 	 */
 	void SetFromPointer(const void* const InData, const int64 InSizeOfT, const int64 InDataSize);
 
 	/**
-	 * Checks and warns whether the current data type T is incompatible with DataType.
+	 * Checks and warns whether the current data type T is the same than the one defined by DataType.
+	 * CheckTAndDataTypeEquivalent() is the function to be used and CheckTAndDataTypeEquivalentAuxiliary its just its auxiliary and moves as much
+	 * code as possible from the template function to the cpp file
 	 */
 	template<typename T>
 	bool CheckTAndDataTypeEquivalent() const;
-	bool CheckTAndDataTypeResult(const bool bInCheckTAndDataTypeResult, const int64 InSizeOfT) const;
+	bool CheckTAndDataTypeEquivalentAuxiliary(const bool bInCheckTAndDataTypeResult, const int64 InSizeOfT) const;
 
 	/**
-	 * This function has total control over the buffer flags (and ignores the TensorType flag). Thus, it is kept private and it should only be used
-	 * by its public analog.
+	 * This function has total control over the buffer flags (and ignores the TensorType flag). Conceptually speaking, it should be a static
+	 * function, but that would add a lot of arguments to it. Thus, it is kept private and it should only be used by its public equivalent.
 	 * Depending on the value of InBufferUsageFlags, it can perform GPU allocation, CPU-to-GPU upload, apply different settings, etc.
 	 * @see CPUToRDGBuilder_RenderThread(FRDGBuilder* InOutGraphBuilder).
 	 */
@@ -510,5 +555,5 @@ void FNeuralTensor::SetTo(const TInput InValue)
 template<typename T>
 bool FNeuralTensor::CheckTAndDataTypeEquivalent() const
 {
-	return CheckTAndDataTypeResult(FNeuralDataTypeUtils::GetDataType<T>() == DataType, sizeof(T));
+	return CheckTAndDataTypeEquivalentAuxiliary(FNeuralDataTypeUtils::GetDataType<T>() == DataType, sizeof(T));
 }
diff --git a/Engine/Plugins/Experimental/NeuralNetworkInference/Source/Private/NeuralNetworkInferenceQA/Private/OperatorTester.cpp b/Engine/Plugins/Experimental/NeuralNetworkInference/Source/Private/NeuralNetworkInferenceQA/Private/OperatorTester.cpp
--- a/Engine/Plugins/Experimental/NeuralNetworkInference/Source/Private/NeuralNetworkInferenceQA/Private/OperatorTester.cpp
+++ b/Engine/Plugins/Experimental/NeuralNetworkInference/Source/Private/NeuralNetworkInferenceQA/Private/OperatorTester.cpp
@@ -10,9 +10,11 @@
 /* FOperatorTester static public functions
  *****************************************************************************/
 
-void FOperatorTester::TestOperator(UNeuralNetworkInferenceQAAsset* InOutNetworkInferenceQAAsset, TArray<TArray<FNeuralTensor>>& InOutTensors, TArray<TArray<FNeuralTensor*>>& InOutInputTensors,
-	TArray<TArray<FNeuralTensor*>>& InOutOutputTensors, TArray<TSharedPtr<FNeuralOperator>>& InOutOperators, const FString& InOperatorName, const FString& InGroundTruthDirectory,
-	const float InZeroThreshold, const bool bInShouldOperatorsProvideSameResult, const bool bInShouldRunGPU, const bool bInShouldRunFinalStringVsIOTest)
+void FOperatorTester::TestOperator(UNeuralNetworkInferenceQAAsset* InOutNetworkInferenceQAAsset, TArray<TArray<FNeuralTensor>>& InOutTensors,
+	TArray<TArray<FNeuralTensor*>>& InOutInputTensors, TArray<TArray<FNeuralTensor*>>& InOutOutputTensors,
+	TArray<TSharedPtr<FNeuralOperator>>& InOutOperators, const FString& InOperatorName, const FString& InGroundTruthDirectory,
+	const float InZeroThreshold, const bool bInShouldOperatorsProvideSameResult, const bool bInShouldRunGPU,
+	const bool bInShouldRunFinalStringVsIOTest)
 {
 	// Sanity checks
 	if (!InOutNetworkInferenceQAAsset)
@@ -27,7 +29,8 @@ void FOperatorTester::TestOperator(UNeuralNetworkInferenceQAAsset* InOutNetworkI
 	}
 	if (!bInShouldRunGPU)
 	{
-		UE_LOG(LogNeuralNetworkInferenceQA, Warning, TEXT("FOperatorTester::TestOperator(): bInShouldRunGPU is false, this should only happen for very early deployment of new operators."));
+		UE_LOG(LogNeuralNetworkInferenceQA, Warning, TEXT("FOperatorTester::TestOperator(): bInShouldRunGPU is false, this should only happen for"
+			" very early deployment of new operators."));
 	}
 	// Run test on each input FNeuralTensor
 	FString Index0StringIfShouldProvideSameResult;
@@ -44,7 +47,8 @@ void FOperatorTester::TestOperator(UNeuralNetworkInferenceQAAsset* InOutNetworkI
 		// Operator array
 		const int32 InlinedTensorIndex = Operator->SetInputTensorsAndGetInlinedIndex(InputTensors);
 		// Tensor manager
-		const TArray<FNeuralTensor*>& OutputTensors = (InlinedTensorIndex > -1 ? TArray<FNeuralTensor*>({ InputTensors[InlinedTensorIndex] }) : InOutOutputTensors[Index]);
+		const TArray<FNeuralTensor*>& OutputTensors = (InlinedTensorIndex > -1
+			? TArray<FNeuralTensor*>({ InputTensors[InlinedTensorIndex] }) : InOutOutputTensors[Index]);
 		if (OutputTensors.IsEmpty())
 		{
 			ensureMsgf(false, TEXT("OutputTensors.Num() is 0 but should be > 0."));
@@ -63,7 +67,8 @@ void FOperatorTester::TestOperator(UNeuralNetworkInferenceQAAsset* InOutNetworkI
 			for (int64 AuxiliaryTensorIndex = 0; AuxiliaryTensorIndex < Operator->GetNumberAuxiliaryTensors(); ++AuxiliaryTensorIndex)
 			{
 				const FString AuxiliaryTensorName = Operator->GetName() + TEXT("-Aux") + FString::FromInt(AuxiliaryTensorIndex);
-				InOutTensors[Index].Push(FNeuralTensor(ENeuralDataType::None, TArray<int64>({}), AuxiliaryTensorName, ENeuralTensorType::Generic)); // ENeuralTensorType can be modify by FNeuralOperator::Configure()
+				InOutTensors[Index].Push(FNeuralTensor(ENeuralDataType::None, TArray<int64>({}),
+					AuxiliaryTensorName, ENeuralTensorType::Generic)); // ENeuralTensorType can be modify by FNeuralOperator::Configure()
 				AuxiliaryOperatorTensors.Push(&InOutTensors[Index].Last());
 			}
 			Operator->SetAuxiliaryTensors(AuxiliaryOperatorTensors);
@@ -111,7 +116,8 @@ void FOperatorTester::TestOperator(UNeuralNetworkInferenceQAAsset* InOutNetworkI
 				const FString OutputTensorsAsString = FNeuralNetworkInferenceQAOperatorAsset::RunNetworkCPUAndGetString(Network);
 				if (Index0StringIfShouldProvideSameResult != OutputTensorsAsString)
 				{
-					UE_LOG(LogNeuralNetworkInferenceQA, Display, TEXT("Index0StringIfShouldProvideSameResult:\n%s"), *Index0StringIfShouldProvideSameResult);
+					UE_LOG(LogNeuralNetworkInferenceQA, Display, TEXT("Index0StringIfShouldProvideSameResult:\n%s"),
+						*Index0StringIfShouldProvideSameResult);
 					UE_LOG(LogNeuralNetworkInferenceQA, Display, TEXT("OutputTensorsAsString:\n%s"), *OutputTensorsAsString);
 					ensureMsgf(false, TEXT("Subtest #%d did not match subtest #0"), Index);
 				}
@@ -125,8 +131,10 @@ void FOperatorTester::TestOperator(UNeuralNetworkInferenceQAAsset* InOutNetworkI
 	}
 }
 
-void FOperatorTester::TestBatchNormalizationOperator(UNeuralNetworkInferenceQAAsset* InOutNetworkInferenceQAAsset, const TMap<int32, TArray<float>>& InInputTensorData, const TMap<int32, TArray<float>>& InInputTensorDataSmall,
-	const TMap<int32, TArray<float>>& InInputTensorDataPos, const float InZeroThreshold, const FString& InGroundTruthDirectory, const int32 InOperatorCounter, const int32 InOperatorTotal)
+void FOperatorTester::TestBatchNormalizationOperator(UNeuralNetworkInferenceQAAsset* InOutNetworkInferenceQAAsset,
+	const TMap<int32, TArray<float>>& InInputTensorData, const TMap<int32, TArray<float>>& InInputTensorDataSmall,
+	const TMap<int32, TArray<float>>& InInputTensorDataPos, const float InZeroThreshold, const FString& InGroundTruthDirectory,
+	const int32 InOperatorCounter, const int32 InOperatorTotal)
 {
 	UE_LOG(LogNeuralNetworkInferenceQA, Display, TEXT("--------------- %d/%d FBatchNormalizationOperator"), InOperatorCounter, InOperatorTotal);
 	// Inout tensors (X != Mean to avoid Output = Bias; Variance positive to avoid NaN)
@@ -145,14 +153,17 @@ void FOperatorTester::TestBatchNormalizationOperator(UNeuralNetworkInferenceQAAs
 			FNeuralTensor(InInputTensorDataPos.FindChecked(3), TArray<int64>({3})),
 			FNeuralTensor(ENeuralDataType::Float, {1,2,1,3})},
 		// Scalars
-		{FNeuralTensor(InInputTensorData.FindChecked(1)), FNeuralTensor(InInputTensorData.FindChecked(1)), FNeuralTensor(InInputTensorData.FindChecked(1)),
-			FNeuralTensor(InInputTensorDataSmall.FindChecked(1)), FNeuralTensor(InInputTensorDataPos.FindChecked(1)), FNeuralTensor(ENeuralDataType::Float, 1)},
+		{FNeuralTensor(InInputTensorData.FindChecked(1)), FNeuralTensor(InInputTensorData.FindChecked(1)),
+			FNeuralTensor(InInputTensorData.FindChecked(1)), FNeuralTensor(InInputTensorDataSmall.FindChecked(1)),
+			FNeuralTensor(InInputTensorDataPos.FindChecked(1)), FNeuralTensor(ENeuralDataType::Float, 1)},
 		// Scalars
-		{FNeuralTensor(InInputTensorData.FindChecked(1)), FNeuralTensor(InInputTensorData.FindChecked(1)), FNeuralTensor(InInputTensorData.FindChecked(1)),
-			FNeuralTensor(InInputTensorDataSmall.FindChecked(1)), FNeuralTensor(InInputTensorDataPos.FindChecked(1)), FNeuralTensor(ENeuralDataType::Float, 1)},
+		{FNeuralTensor(InInputTensorData.FindChecked(1)), FNeuralTensor(InInputTensorData.FindChecked(1)),
+			FNeuralTensor(InInputTensorData.FindChecked(1)), FNeuralTensor(InInputTensorDataSmall.FindChecked(1)),
+			FNeuralTensor(InInputTensorDataPos.FindChecked(1)), FNeuralTensor(ENeuralDataType::Float, 1)},
 		// 1-D
-		{FNeuralTensor(InInputTensorData.FindChecked(10)), FNeuralTensor(InInputTensorData.FindChecked(1)), FNeuralTensor(InInputTensorData.FindChecked(1)),
-			FNeuralTensor(InInputTensorDataSmall.FindChecked(1)), FNeuralTensor(InInputTensorDataPos.FindChecked(1)), FNeuralTensor(ENeuralDataType::Float, 10)},
+		{FNeuralTensor(InInputTensorData.FindChecked(10)), FNeuralTensor(InInputTensorData.FindChecked(1)),
+			FNeuralTensor(InInputTensorData.FindChecked(1)), FNeuralTensor(InInputTensorDataSmall.FindChecked(1)),
+			FNeuralTensor(InInputTensorDataPos.FindChecked(1)), FNeuralTensor(ENeuralDataType::Float, 10)},
 		// 4-D with 1 batch and 1 channel
 		{FNeuralTensor(InInputTensorData.FindChecked(100), TArray<int64>({1, 1, 5, 20})), FNeuralTensor(InInputTensorData.FindChecked(1)),
 			FNeuralTensor(InInputTensorData.FindChecked(1)), FNeuralTensor(InInputTensorDataSmall.FindChecked(1)),
@@ -187,36 +198,45 @@ void FOperatorTester::TestBatchNormalizationOperator(UNeuralNetworkInferenceQAAs
 	TArray<TSharedPtr<FNeuralOperator>> Operators;
 	for (int32 TensorIndex = 0; TensorIndex < Tensors.Num(); ++TensorIndex)
 	{
-		InputTensors.Push({ &Tensors[TensorIndex][0], &Tensors[TensorIndex][1], &Tensors[TensorIndex][2], &Tensors[TensorIndex][3], &Tensors[TensorIndex][4] });
+		InputTensors.Push({ &Tensors[TensorIndex][0], &Tensors[TensorIndex][1], &Tensors[TensorIndex][2], &Tensors[TensorIndex][3],
+			&Tensors[TensorIndex][4] });
 		OutputTensors.Push(TensorIndex % 2 == 0 ? TArray<FNeuralTensor*>({ &Tensors[TensorIndex][5] }) : TArray<FNeuralTensor*>({/*Inlined*/ }));
 		Operators.Push(MakeShared<FBatchNormalizationOperator>(OutputTensors[TensorIndex].Num() == 0,
 			/*Epsilon*/(TensorIndex == 1 ? 1e-2 : TensorIndex == 3 ? 10 : TensorIndex == 6 ? 1 : 1e-5),
 			/*Momentum*/(TensorIndex+1.f)/Tensors.Num()));
 	}
-	TestOperator(InOutNetworkInferenceQAAsset, Tensors, InputTensors, OutputTensors, Operators, TEXT("FBatchNormalizationOperator"), InGroundTruthDirectory, InZeroThreshold,
-		/*bShouldOperatorsProvideSameResult*/false, /*bInShouldRunGPU*/true);
+	TestOperator(InOutNetworkInferenceQAAsset, Tensors, InputTensors, OutputTensors, Operators, TEXT("FBatchNormalizationOperator"),
+		InGroundTruthDirectory, InZeroThreshold, /*bShouldOperatorsProvideSameResult*/false, /*bInShouldRunGPU*/true);
 }
 
-void FOperatorTester::TestGemmOperator(UNeuralNetworkInferenceQAAsset* InOutNetworkInferenceQAAsset, const TMap<int32, TArray<float>>& InInputTensorData, const float InZeroThreshold, const FString& InGroundTruthDirectory,
+void FOperatorTester::TestGemmOperator(UNeuralNetworkInferenceQAAsset* InOutNetworkInferenceQAAsset,
+	const TMap<int32, TArray<float>>& InInputTensorData, const float InZeroThreshold, const FString& InGroundTruthDirectory,
 	const int32 InOperatorCounter, const int32 InOperatorTotal)
 {
 	UE_LOG(LogNeuralNetworkInferenceQA, Display, TEXT("--------------- %d/%d FGemmOperator"), InOperatorCounter, InOperatorTotal);
 	// Inout tensors
 	const TArray<TArray<FNeuralTensor>> GemmABCOutputTensorTests({
 		// Scalars
-		{FNeuralTensor(InInputTensorData.FindChecked(1)), FNeuralTensor(InInputTensorData.FindChecked(1)), FNeuralTensor(InInputTensorData.FindChecked(1)),
-			FNeuralTensor(ENeuralDataType::Float, 1)},
+		{FNeuralTensor(InInputTensorData.FindChecked(1)), FNeuralTensor(InInputTensorData.FindChecked(1)),
+			FNeuralTensor(InInputTensorData.FindChecked(1)), FNeuralTensor(ENeuralDataType::Float, 1)},
 		// 1-D
-		{FNeuralTensor(InInputTensorData.FindChecked(300), TArray<int64>({50, 6})), FNeuralTensor(InInputTensorData.FindChecked(18), TArray<int64>({6, 3})),
-			FNeuralTensor(InInputTensorData.FindChecked(150), TArray<int64>({50, 3})), FNeuralTensor(ENeuralDataType::Float, TArray<int64>({50, 3}))},
+		{FNeuralTensor(InInputTensorData.FindChecked(300), TArray<int64>({50, 6})),
+			FNeuralTensor(InInputTensorData.FindChecked(18), TArray<int64>({6, 3})),
+			FNeuralTensor(InInputTensorData.FindChecked(150), TArray<int64>({50, 3})),
+			FNeuralTensor(ENeuralDataType::Float, TArray<int64>({50, 3}))},
 		// 1-D, C size = [(1), 3], the next 2 should give the exact same results
-		{FNeuralTensor(InInputTensorData.FindChecked(300), TArray<int64>({50, 6})), FNeuralTensor(InInputTensorData.FindChecked(18), TArray<int64>({6, 3})),
-			FNeuralTensor(InInputTensorData.FindChecked(3), TArray<int64>({1, 3})), FNeuralTensor(ENeuralDataType::Float, TArray<int64>({50, 3}))}, // C size = [1, 3]
-		{FNeuralTensor(InInputTensorData.FindChecked(300), TArray<int64>({50, 6})), FNeuralTensor(InInputTensorData.FindChecked(18), TArray<int64>({6, 3})),
+		{FNeuralTensor(InInputTensorData.FindChecked(300), TArray<int64>({50, 6})),
+			FNeuralTensor(InInputTensorData.FindChecked(18), TArray<int64>({6, 3})),
+			FNeuralTensor(InInputTensorData.FindChecked(3), TArray<int64>({1, 3})),
+			FNeuralTensor(ENeuralDataType::Float, TArray<int64>({50, 3}))}, // C size = [1, 3]
+		{FNeuralTensor(InInputTensorData.FindChecked(300), TArray<int64>({50, 6})),
+			FNeuralTensor(InInputTensorData.FindChecked(18), TArray<int64>({6, 3})),
 			FNeuralTensor(InInputTensorData.FindChecked(3)), FNeuralTensor(ENeuralDataType::Float, TArray<int64>({50, 3}))}, // C size = [3]
 		// 1-D, C size = [50, 1]
-		{FNeuralTensor(InInputTensorData.FindChecked(300), TArray<int64>({50, 6})), FNeuralTensor(InInputTensorData.FindChecked(18), TArray<int64>({6, 3})),
-			FNeuralTensor(InInputTensorData.FindChecked(50), TArray<int64>({50, 1})), FNeuralTensor(ENeuralDataType::Float, TArray<int64>({50, 3}))}
+		{FNeuralTensor(InInputTensorData.FindChecked(300), TArray<int64>({50, 6})),
+			FNeuralTensor(InInputTensorData.FindChecked(18), TArray<int64>({6, 3})),
+			FNeuralTensor(InInputTensorData.FindChecked(50), TArray<int64>({50, 1})),
+			FNeuralTensor(ENeuralDataType::Float, TArray<int64>({50, 3}))}
 	});
 	for (int32 TensorTestsIndex = 0; TensorTestsIndex < GemmABCOutputTensorTests.Num(); ++TensorTestsIndex)
 	{
@@ -263,7 +283,8 @@ void FOperatorTester::TestGemmOperator(UNeuralNetworkInferenceQAAsset* InOutNetw
 			}
 			else
 			{
-				InputTensors.Push({ &GemmABCOutputTensors[InoutIndex][0], &GemmABCOutputTensors[InoutIndex][1], &GemmABCOutputTensors[InoutIndex][2] });
+				InputTensors.Push({ &GemmABCOutputTensors[InoutIndex][0], &GemmABCOutputTensors[InoutIndex][1],
+					&GemmABCOutputTensors[InoutIndex][2] });
 			}
 			// Transpose FNeuralTensor if bTransA or bTransB
 			if ((InoutIndex >= 4 && InoutIndex < 8) || InoutIndex >= 12) // bTransA?
@@ -283,13 +304,15 @@ void FOperatorTester::TestGemmOperator(UNeuralNetworkInferenceQAAsset* InOutNetw
 			OutputTensors.Push({ &GemmABCOutputTensors[InoutIndex][3] });
 		}
 		// Run actual layer unit testing
-		TestOperator(InOutNetworkInferenceQAAsset, GemmABCOutputTensors, InputTensors, OutputTensors, Operators, /*FileName*/TEXT("FGemmOperator"), InGroundTruthDirectory, InZeroThreshold,
-			/*bShouldOperatorsProvideSameResult*/false, /*bShouldRunGPU*/true, (TensorTestsIndex+1 == GemmABCOutputTensorTests.Num()));
+		TestOperator(InOutNetworkInferenceQAAsset, GemmABCOutputTensors, InputTensors, OutputTensors, Operators,
+			/*FileName*/TEXT("FGemmOperator"), InGroundTruthDirectory, InZeroThreshold, /*bShouldOperatorsProvideSameResult*/false,
+			/*bShouldRunGPU*/true, (TensorTestsIndex+1 == GemmABCOutputTensorTests.Num()));
 	}
 }
 
-void FOperatorTester::TestConvOperator(UNeuralNetworkInferenceQAAsset* InOutNetworkInferenceQAAsset, const TMap<int32, TArray<float>>& InInputTensorData, const TMap<int32, TArray<float>>& InInputTensorDataSmall,
-	const float InZeroThreshold, const FString& InGroundTruthDirectory, const int32 InOperatorCounter, const int32 InOperatorTotal)
+void FOperatorTester::TestConvOperator(UNeuralNetworkInferenceQAAsset* InOutNetworkInferenceQAAsset,
+	const TMap<int32, TArray<float>>& InInputTensorData, const TMap<int32, TArray<float>>& InInputTensorDataSmall, const float InZeroThreshold,
+	const FString& InGroundTruthDirectory, const int32 InOperatorCounter, const int32 InOperatorTotal)
 {
 	UE_LOG(LogNeuralNetworkInferenceQA, Display, TEXT("--------------- %d/%d FConvOperator"), InOperatorCounter, InOperatorTotal);
 	// Inout tensors
@@ -964,28 +987,32 @@ void FOperatorTester::TestConvTransposeOperator(UNeuralNetworkInferenceQAAsset*
 		// Create and push FConvTransposeOperator
 		{
 			TSharedPtr<FConvTransposeOperator> ConvTransposeOperator;
-			ConvTransposeOperator = MakeShared<FConvTransposeOperator>(AutoPad, Dilations, Group, /*KernelShape*/TArray<int64>(), OutputPadding, OutputShape, Pads, Strides);
+			ConvTransposeOperator = MakeShared<FConvTransposeOperator>(AutoPad, Dilations, Group, /*KernelShape*/TArray<int64>(), OutputPadding,
+				OutputShape, Pads, Strides);
 			// Flip weights
 			{
-				UE_LOG(LogNeuralNetworkInferenceQA, Display, TEXT("Flipping FConvTranspose weights. This is a temporary measurement but it will keep the weights flipped for now (including in the unit test txt files)."));
+				UE_LOG(LogNeuralNetworkInferenceQA, Display, TEXT("Flipping FConvTranspose weights. This is a temporary measurement but it will keep"
+					" the weights flipped for now (including in the unit test txt files)."));
 				ConvTransposeOperator->SetWhetherWeightTensorWasFlipped(true);
 				TensorsIndex[1].Flip(2, TensorsIndex[1].GetNumberDimensions());
 			}
 			Operators.Emplace(ConvTransposeOperator);
 		}
-		// Remove all non-required/auxiliary tensors from InputTensors to keep it in X, W, and B (if B exists). Otherwise, those would be used by the network and moved into GPU memory
+		// Remove all non-required/auxiliary tensors from InputTensors to keep it in X, W, and B (if B exists). Otherwise, those would be used by the
+		// network and moved into GPU memory
 		while (Tensors[TensorIndex].Num() > 4 || Tensors[TensorIndex].Last().Num() < 1)
 		{
 			Tensors[TensorIndex].Pop();
 		}
 		// TensorsIndex is no longer valid after this while loop!
 	}
-	TestOperator(InOutNetworkInferenceQAAsset, Tensors, InputTensors, OutputTensors, Operators, TEXT("FConvTransposeOperator"), InGroundTruthDirectory, InZeroThreshold,
-		/*bShouldOperatorsProvideSameResult*/false, /*bInShouldRunGPU*/true);
+	TestOperator(InOutNetworkInferenceQAAsset, Tensors, InputTensors, OutputTensors, Operators, TEXT("FConvTransposeOperator"),
+		InGroundTruthDirectory, InZeroThreshold, /*bShouldOperatorsProvideSameResult*/false, /*bInShouldRunGPU*/true);
 }
 
-void FOperatorTester::TestReshapeOperator(UNeuralNetworkInferenceQAAsset* InOutNetworkInferenceQAAsset, const TMap<int32, TArray<float>>& InInputTensorData,
-	const TMap<int32, TArray<float>>& InInputTensorDataPos, const float InZeroThreshold, const FString& InGroundTruthDirectory, const int32 InOperatorCounter, const int32 InOperatorTotal)
+void FOperatorTester::TestReshapeOperator(UNeuralNetworkInferenceQAAsset* InOutNetworkInferenceQAAsset,
+	const TMap<int32, TArray<float>>& InInputTensorData, const TMap<int32, TArray<float>>& InInputTensorDataPos, const float InZeroThreshold,
+	const FString& InGroundTruthDirectory, const int32 InOperatorCounter, const int32 InOperatorTotal)
 {
 	UE_LOG(LogNeuralNetworkInferenceQA, Display, TEXT("--------------- %d/%d FReshapeOperator"), InOperatorCounter, InOperatorTotal);
 	// Inout tensors (X != Mean to avoid Output = Bias; Variance positive to avoid NaN)
