Adding NNI documentation:
- NNI/Public/NeuraTensor.h fully cleaned and documented for GPU as well.

#rnx

diff --git a/Engine/Plugins/Experimental/NeuralNetworkInference/Source/NeuralNetworkInference/Private/NeuralNetworkImplBackEndUEAndORT.cpp b/Engine/Plugins/Experimental/NeuralNetworkInference/Source/NeuralNetworkInference/Private/NeuralNetworkImplBackEndUEAndORT.cpp
--- a/Engine/Plugins/Experimental/NeuralNetworkInference/Source/NeuralNetworkInference/Private/NeuralNetworkImplBackEndUEAndORT.cpp
+++ b/Engine/Plugins/Experimental/NeuralNetworkInference/Source/NeuralNetworkInference/Private/NeuralNetworkImplBackEndUEAndORT.cpp
@@ -792,7 +792,7 @@ bool UNeuralNetwork::FImplBackEndUEAndORT::SetTensorsFromNetwork(TArray<FNeuralT
 			// @todo: This requires SetNumUnitialized() to be run, otherwise Size and Volume will be set to 0
 			void* D3DResource = nullptr;
 
-			if (!OutTensors[TensorIndex].InitPooledBuffer(&D3DResource))
+			if (!OutTensors[TensorIndex].InitPooledBufferForUEAndORTBackEnd(&D3DResource))
 			{
 				UE_LOG(LogNeuralNetworkInference, Warning,
 					TEXT("FImplBackEndUEAndORT::SetTensorsFromNetwork(): Failed to initialize the pooled buffer."));
diff --git a/Engine/Plugins/Experimental/NeuralNetworkInference/Source/NeuralNetworkInference/Private/NeuralTensor.cpp b/Engine/Plugins/Experimental/NeuralNetworkInference/Source/NeuralNetworkInference/Private/NeuralTensor.cpp
--- a/Engine/Plugins/Experimental/NeuralNetworkInference/Source/NeuralNetworkInference/Private/NeuralTensor.cpp
+++ b/Engine/Plugins/Experimental/NeuralNetworkInference/Source/NeuralNetworkInference/Private/NeuralTensor.cpp
@@ -607,11 +607,11 @@ void FNeuralTensor::RDGBuilderToCPU_RenderThread(FRDGBuilder* InOutGraphBuilder)
 	});
 }
 
-bool FNeuralTensor::InitPooledBuffer(void** InOutNativeResource)
+bool FNeuralTensor::InitPooledBufferForUEAndORTBackEnd(void** InOutNativeResource)
 {
 	if (!bEnableGPU)
 	{
-		UE_LOG(LogNeuralNetworkInference, Warning, TEXT("FNeuralTensor::InitPooledBuffer(): bEnableGPU is false."));
+		UE_LOG(LogNeuralNetworkInference, Warning, TEXT("FNeuralTensor::InitPooledBufferForUEAndORTBackEnd(): bEnableGPU is false."));
 		return false;
 	}
 
diff --git a/Engine/Plugins/Experimental/NeuralNetworkInference/Source/NeuralNetworkInference/Public/NeuralTensor.h b/Engine/Plugins/Experimental/NeuralNetworkInference/Source/NeuralNetworkInference/Public/NeuralTensor.h
--- a/Engine/Plugins/Experimental/NeuralNetworkInference/Source/NeuralNetworkInference/Public/NeuralTensor.h
+++ b/Engine/Plugins/Experimental/NeuralNetworkInference/Source/NeuralNetworkInference/Public/NeuralTensor.h
@@ -312,7 +312,7 @@ public:
 	 * @param InOutNativeResource Pointer to the ID3D12Resource obtained from PooledBuffer. This resource can be shared with the DirectML execution
 	 * provider.
 	 */
-	bool InitPooledBuffer(void** InOutNativeResource = nullptr);
+	bool InitPooledBufferForUEAndORTBackEnd(void** InOutNativeResource = nullptr);
 
 	/**
 	 * It returns a constant reference of PooledBuffer. @see PooledBuffer for more details.
