NeuralNetworkInference plugin:
Setting EnablePattern in CPU and disable for DML. Remove runningtime warning:
[W:onnxruntime:, inference_session.cc:489 onnxruntime::InferenceSession::RegisterExecutionProvider] Having memory pattern enabled is not supported while using the DML Execution Provider. So disabling it for this session since it uses the DML Execution Provider.

diff --git a/Engine/Plugins/Experimental/NeuralNetworkInference/Source/NeuralNetworkInference/Private/NeuralNetworkImplBackEndUEAndORT.cpp b/Engine/Plugins/Experimental/NeuralNetworkInference/Source/NeuralNetworkInference/Private/NeuralNetworkImplBackEndUEAndORT.cpp
--- a/Engine/Plugins/Experimental/NeuralNetworkInference/Source/NeuralNetworkInference/Private/NeuralNetworkImplBackEndUEAndORT.cpp
+++ b/Engine/Plugins/Experimental/NeuralNetworkInference/Source/NeuralNetworkInference/Private/NeuralNetworkImplBackEndUEAndORT.cpp
@@ -569,6 +569,9 @@ bool UNeuralNetwork::FImplBackEndUEAndORT::ConfigureMembers(const ENeuralDeviceT
 		// ORT GPU (Direct ML)
 		SessionOptions->SetGraphOptimizationLevel(GraphOptimizationLevel::ORT_ENABLE_ALL); // ORT_ENABLE_ALL, ORT_ENABLE_EXTENDED, ORT_ENABLE_BASIC, ORT_DISABLE_ALL
 
+		// DML's memory is not byte addressable and hence mem pattern doesn't work.
+		SessionOptions->DisableCpuMemArena();
+
 		// Get DML API
 		const OrtApi* OrtApi = OrtGetApiBase()->GetApi(ORT_API_VERSION);
 		
@@ -614,6 +617,7 @@ bool UNeuralNetwork::FImplBackEndUEAndORT::ConfigureMembers(const ENeuralDeviceT
 #else
 		// ORT CPU
 		SessionOptions->SetGraphOptimizationLevel(GraphOptimizationLevel::ORT_ENABLE_ALL); // ORT_ENABLE_ALL, ORT_ENABLE_EXTENDED, ORT_ENABLE_BASIC, ORT_DISABLE_ALL
+		SessionOptions->EnableCpuMemArena();
 #endif //ORT_CPU
 	}
 
