Enabling NNI by default.

#rnx

diff --git a/Engine/Plugins/Experimental/NeuralNetworkInference/NeuralNetworkInference.uplugin b/Engine/Plugins/Experimental/NeuralNetworkInference/NeuralNetworkInference.uplugin
--- a/Engine/Plugins/Experimental/NeuralNetworkInference/NeuralNetworkInference.uplugin
+++ b/Engine/Plugins/Experimental/NeuralNetworkInference/NeuralNetworkInference.uplugin
@@ -10,7 +10,7 @@
 	"DocsURL": "",
 	"MarketplaceURL": "",
 	"SupportURL": "",
-	"EnabledByDefault": false,
+	"EnabledByDefault": true,
 	"CanContainContent": false,
 	"IsBetaVersion": true,
 	"IsExperimentalVersion": true,
