Adding NNI documentation:
- NNI/Public/NeuraNetwork.h partially cleaned and documented
U* --> UObjectPtr

#rnx

diff --git a/Engine/Plugins/Experimental/NeuralNetworkInference/Source/NeuralNetworkInference/Private/NeuralNetwork.cpp b/Engine/Plugins/Experimental/NeuralNetworkInference/Source/NeuralNetworkInference/Private/NeuralNetwork.cpp
--- a/Engine/Plugins/Experimental/NeuralNetworkInference/Source/NeuralNetworkInference/Private/NeuralNetwork.cpp
+++ b/Engine/Plugins/Experimental/NeuralNetworkInference/Source/NeuralNetworkInference/Private/NeuralNetwork.cpp
@@ -49,17 +49,13 @@ UNeuralNetwork::UNeuralNetwork()
 	, InputDeviceType(ENeuralDeviceType::CPU)
 	, OutputDeviceType(ENeuralDeviceType::CPU)
 	, SynchronousMode(ENeuralSynchronousMode::Synchronous)
-	, DelegateThreadMode(ENeuralThreadMode::GameThread)
+	, ThreadModeDelegateForAsyncRunCompleted(ENeuralThreadMode::GameThread)
 	, bIsLoaded(false)
 	, BackEnd(ENeuralBackEnd::Auto)
 	, BackEndForCurrentPlatform(FPrivateNeuralNetwork::SetBackEndForCurrentPlatform(BackEnd))
 {
 }
 
-UNeuralNetwork::~UNeuralNetwork()
-{
-}
-
 
 
 /* UNeuralNetwork public functions
@@ -176,15 +172,15 @@ UNeuralNetwork::FOnAsyncRunCompleted& UNeuralNetwork::GetOnAsyncRunCompletedDele
 	return OnAsyncRunCompletedDelegate;
 }
 
-ENeuralThreadMode UNeuralNetwork::GetOnAsyncRunCompletedDelegateMode() const
+ENeuralThreadMode UNeuralNetwork::GetThreadModeDelegateForAsyncRunCompleted() const
 {
-	return DelegateThreadMode;
+	return ThreadModeDelegateForAsyncRunCompleted;
 }
 
-void UNeuralNetwork::SetOnAsyncRunCompletedDelegateMode(const ENeuralThreadMode InDelegateThreadMode)
+void UNeuralNetwork::SetThreadModeDelegateForAsyncRunCompleted(const ENeuralThreadMode InThreadModeDelegateForAsyncRunCompleted)
 {
 	const FScopeLock ResourcesLock(&ResoucesCriticalSection);
-	DelegateThreadMode = InDelegateThreadMode;
+	ThreadModeDelegateForAsyncRunCompleted = InThreadModeDelegateForAsyncRunCompleted;
 }
 
 void UNeuralNetwork::ResetStats()
@@ -539,7 +535,7 @@ bool UNeuralNetwork::Load()
 	if (BackEndForCurrentPlatform == ENeuralBackEnd::UEAndORT)
 	{
 		UNeuralNetwork::FImplBackEndUEAndORT::WarnAndSetDeviceToCPUIfDX12NotEnabled(DeviceType, /*bShouldOpenMessageLog*/true);
-		bIsLoaded = UNeuralNetwork::FImplBackEndUEAndORT::Load(ImplBackEndUEAndORT, OnAsyncRunCompletedDelegate, DelegateThreadMode, ResoucesCriticalSection, AreInputTensorSizesVariable,
+		bIsLoaded = UNeuralNetwork::FImplBackEndUEAndORT::Load(ImplBackEndUEAndORT, OnAsyncRunCompletedDelegate, ThreadModeDelegateForAsyncRunCompleted, ResoucesCriticalSection, AreInputTensorSizesVariable,
 			ModelReadFromFileInBytes, ModelFullFilePath, GetDeviceType(), GetInputDeviceType(), GetOutputDeviceType());
 	}
 	// UEOnly
@@ -617,12 +613,12 @@ bool UNeuralNetwork::SetBackEnd(const UNeuralNetwork::ENeuralBackEnd InBackEnd)
 }
 
 #if WITH_EDITOR
-UAssetImportData* UNeuralNetwork::GetAssetImportData() const
+TObjectPtr<UAssetImportData> UNeuralNetwork::GetAssetImportData() const
 {
 	return AssetImportData;
 }
 
-UAssetImportData* UNeuralNetwork::GetAndMaybeCreateAssetImportData()
+TObjectPtr<UAssetImportData> UNeuralNetwork::GetAndMaybeCreateAssetImportData()
 {
 	// An existing import data object was not found, so make one here.
 	if (!AssetImportData)
diff --git a/Engine/Plugins/Experimental/NeuralNetworkInference/Source/NeuralNetworkInference/Public/NeuralEnumClasses.h b/Engine/Plugins/Experimental/NeuralNetworkInference/Source/NeuralNetworkInference/Public/NeuralEnumClasses.h
--- a/Engine/Plugins/Experimental/NeuralNetworkInference/Source/NeuralNetworkInference/Public/NeuralEnumClasses.h
+++ b/Engine/Plugins/Experimental/NeuralNetworkInference/Source/NeuralNetworkInference/Public/NeuralEnumClasses.h
@@ -54,7 +54,7 @@ enum class ENeuralSynchronousMode : uint8
 	 * (e.g., UNeuralNetwork::GetOnAsyncRunCompletedDelegate() for asynchronous UNeuralNetwork::Run()).
 	 *
 	 * Very important: It takes ~1 millisecond to start the background thread. If your operation (e.g., your network inference time) runs
-	 * synchronously faster than 1-2 msec, using asynchronous is not recommended because it will slow down both the main and bakground threads.
+	 * synchronously faster than 1-2 milliseconds, using asynchronous is not recommended because it'd slow down both the main and background threads.
 	 */
 	Asynchronous
 };
diff --git a/Engine/Plugins/Experimental/NeuralNetworkInference/Source/NeuralNetworkInference/Public/NeuralNetwork.h b/Engine/Plugins/Experimental/NeuralNetworkInference/Source/NeuralNetworkInference/Public/NeuralNetwork.h
--- a/Engine/Plugins/Experimental/NeuralNetworkInference/Source/NeuralNetworkInference/Public/NeuralNetwork.h
+++ b/Engine/Plugins/Experimental/NeuralNetworkInference/Source/NeuralNetworkInference/Public/NeuralNetwork.h
@@ -4,35 +4,45 @@
 
 #include "CoreMinimal.h"
 #include "NeuralEnumClasses.h"
-#include "NeuralTensor.h"
 #include "NeuralStats.h"
+#include "NeuralTensor.h"
 #include "NeuralNetwork.generated.h"
 
 /**
- * UNeuralNetwork is UE's representation for deep learning and neural network models. It supports the industry standard ONNX model format.
- * All major frameworks (PyTorch, TensorFlow, MXNet, Caffe2, etc.) provide converters to ONNX.
+ * NeuralNetworkInference (NNI) is Unreal Engine's framework for running deep learning and neural network inference. It is focused on:
+ * - Efficiency: Underlying state-of-the-art accelerators (DirectML, AVX, CoreML, etc).
+ * - Ease-of-use: Simple but powerful API.
+ * - Completeness: All the functionality of any state-of-the-art deep learning framework.
+ *
+ * UNeuralNetwork is the key class of NNI, and the only one users should interact with. It represents the deep neural model itself. It is capable of
+ * loading and running inference (i.e., a forward pass) on any ONNX (Open Neural Network eXchange) model. ONNX is the industry standard for ML
+ * interoperability, and all major frameworks (PyTorch, TensorFlow, MXNet, Caffe2, etc.) provide converters to ONNX.
  *
- * See the following examples to learn how to read any ONNX model and run inference (i.e., a forward pass) on it.
- * 1. Constructing a UNeuralNetwork from an ONNX file (Editor-only):
+ * The following code snippets show the UNeuralNetwork basics (reading a ONNX model and running inference on it). For more detailed examples, see
+ * {UE5}/Samples/MachineLearning/NNI.
+ *
+ * 1.a. Creating a new UNeuralNetwork and loading a network from an ONNX file:
  *		// Create the UNeuralNetwork object
  *		UNeuralNetwork* Network = NewObject<UNeuralNetwork>((UObject*)GetTransientPackage(), UNeuralNetwork::StaticClass());
- *		// Try to load the network and set the device (CPU/GPU)
+ *		// Load the ONNX model and set the device (CPU/GPU)
  *		const FString ONNXModelFilePath = TEXT("SOME_PARENT_FOLDER/SOME_ONNX_FILE_NAME.onnx");
  *		if (Network->Load(ONNXModelFilePath))
  *		{
  *			Network->SetDeviceType(ENeuralDeviceType::CPU); // Set to CPU/GPU mode
  *		}
  *
- * 2. Loading a UNeuralNetwork from a previously-created UAsset (in Editor or in Game):
- *		// Create and load the UNeuralNetwork object from a UAsset
- *		const FString NetworkUAssetFilePath = TEXT("ExampleNetwork'/Game/Models/ExampleNetwork/ExampleNetwork.ExampleNetwork'");
+ * 1.b. Loading a UNeuralNetwork from a previously-created UNeuralNetwork UAsset:
+ *		// Load the UNeuralNetwork object from a UNeuralNetwork UAsset
+ *		const FString NetworkUAssetFilePath = TEXT("'/Game/Models/ExampleNetwork/ExampleNetwork.ExampleNetwork'");
  *		UNeuralNetwork* Network = LoadObject<UNeuralNetwork>((UObject*)GetTransientPackage(), *NetworkUAssetFilePath);
  *		// Check that the network was successfully loaded
  *		check(Network->IsLoaded());
+ *      // Optionally set to CPU/GPU mode. This step is optional because otherwise it will use the device type that was saved on the loaded UAsset
+ *      // Network->SetDeviceType(ENeuralDeviceType::CPU);
  *
- * 3.1. Running inference (i.e., a forward pass):
+ * 2.a. Running inference:
  *		// Fill input neural tensor
- *		TArray<float> InArray;
+ *		TArray<float> InArray; // Assumed initialized with data and that InArray.Num() == Network->Num()
  *		Network->SetInputFromArrayCopy(InArray);
  *		UE_LOG(LogNeuralNetworkInference, Display, TEXT("Input tensor: %s."), *Network->GetInputTensor().ToString());
  *		// Run UNeuralNetwork
@@ -41,17 +51,21 @@
  *		const FNeuralTensor& OutputTensor = Network->GetOutputTensor();
  *		UE_LOG(LogNeuralNetworkInference, Display, TEXT("Output tensor: %s."), *OutputTensor.ToString());
  *
- * 3.2. Alternative - Filling the input tensor without a TArray-to-FNeuralTensor copy:
+ * 2.b. Running inference more efficient alternative - Filling the input tensor without a TArray-to-FNeuralTensor copy:
  *		// Obtain input tensor pointer
  *		float* InputDataPointer = Network->GetInputDataPointerMutable<float>();
  *		// Fill InputDataPointer
  *		for (int64 Index = 0; Index < Network->GetInputTensor().Num(); ++Index)
- *			InputDataPointer[Index] = ...;
+ *			InputDataPointer[Index] = ...; // Assumed some preprocessing or otherwise simply use Memcpy
  *
- * 3.3. Alternative - Networks with multiple input/output tensors:
- * - Multiple inputs: Add InTensorIndex to GetInputTensor(InTensorIndex) or GetInputDataPointerMutable(InTensorIndex) in the examples above, or use
- *   GetInputTensors() instead.
- * - Multiple outputs: Add InTensorIndex to GetOutputTensor(InTensorIndex) in the examples above or use GetOutputTensors() instead.
+ * 3. Networks with multiple input/output tensors:
+ *      - Multiple inputs: Add InTensorIndex to SetInputFromArrayCopy, GetInputTensor(InTensorIndex) or GetInputDataPointerMutable(InTensorIndex) in
+ *        the examples above:
+ *		    Network->SetInputFromArrayCopy(InputArray0, 0);
+ *		    Network->SetInputFromArrayCopy(InputArray1, 1);
+ *      - Multiple outputs: Add InTensorIndex to GetOutputTensor(InTensorIndex) in the examples above or use GetOutputTensors() instead.
+ *		    const FNeuralTensor& OutputTensor0 = Network->GetOutputTensor(0);
+ *		    const FNeuralTensor& OutputTensor1 = Network->GetOutputTensor(1);
  */
 UCLASS(BlueprintType)
 class NEURALNETWORKINFERENCE_API UNeuralNetwork : public UObject
@@ -59,32 +73,40 @@ class NEURALNETWORKINFERENCE_API UNeuralNetwork : public UObject
 	GENERATED_BODY()
 
 public:
+	/**
+	 * Default constructor that initializes the internal class variables.
+	 */
 	UNeuralNetwork();
 
-	virtual ~UNeuralNetwork();
-
 	/**
-	 * It loads the desired network graph definition and weights given an input ONNX file path.
-	 * @param InModelFilePath can either be a full path or a relative path with respect to the Game project.
-	 * @return Whether the network was successfully loaded.
+	 * Destructor defined in case this class is ever inherited.
 	 */
-	bool Load(const FString& InModelFilePath);
+	virtual ~UNeuralNetwork() = default;
 
 	/**
-	 * It loads the desired network graph definition and weights given an input ONNX file path that has been read as a TArray<uint8> buffer.
-	 * @see FFileHelper::LoadFileToArray for an example of how to read a file into a TArray<uint8>.
-	 * @param InModelReadFromFileInBytes will be moved for performance reasons.
-	 * @return Whether the network was successfully loaded.
+	 * It loads the desired network graph definition and weights from an input ONNX file. This file can be passed either a file path or as a memory
+	 * buffer.
+	 * @param InModelFilePath Input ONNX file path to be read. It can either be a full path or a relative path with respect to the Engine or Game
+	 * project.
+	 * @param InModelReadFromFileInBytes TArray buffer filled with the contents of the ONNX file. This TArray<uint8> buffer will be moved for
+	 * performance reasons. @see FFileHelper::LoadFileToArray for an example of how to read a file into a TArray<uint8>.
+	 * @return Whether the network was successfully loaded. Equivalent to IsLoaded(). @see IsLoaded() for more details.
 	 */
+	bool Load(const FString& InModelFilePath);
 	bool Load(TArray<uint8>& InModelReadFromFileInBytes);
 
 	/**
-	 * It returns whether a network is currently loaded.
+	 * It returns whether a network is currently loaded. It is equivalent to the output of Load().
+	 * @see Load() for more details.
 	 */
 	bool IsLoaded() const;
 
 	/**
-	 * Getter and setter functions for DeviceType, InputDeviceType, and OutputDeviceType.
+	 * Getter and setter functions for DeviceType, InputDeviceType, and OutputDeviceType:
+	 * - GetDeviceType(): It returns DeviceType.
+	 * - GetInputDeviceType(): It returns InputDeviceType.
+	 * - GetOutputDeviceType(): It returns OutputDeviceType.
+	 * - SetDeviceType(): It sets DeviceType, InputDeviceType, and OutputDeviceType.
 	 * @see DeviceType, InputDeviceType, and OutputDeviceType for more details.
 	 */
 	ENeuralDeviceType GetDeviceType() const;
@@ -94,35 +116,54 @@ public:
 		const ENeuralDeviceType InOutputDeviceType = ENeuralDeviceType::CPU);
 
 	/**
-	 * Getter and setter functions for SynchronousMode.
+	 * Getter and setter functions for SynchronousMode:
+	 * - GetDeviceType(): It returns SynchronousMode.
+	 * - SetDeviceType(): It sets SynchronousMode.
 	 * @see SynchronousMode and GetOnAsyncRunCompletedDelegate() for more details.
 	 */
 	ENeuralSynchronousMode GetSynchronousMode() const;
 	void SetSynchronousMode(const ENeuralSynchronousMode InSynchronousMode);
 
 	/**
-	 * GetOnAsyncRunCompletedDelegate() returns a FOnAsyncRunCompleted delegate that will be called when async UNeuralNetwork::Run() is completed.
-	 * - If SynchronousMode == ENeuralSynchronousMode::Asynchronous, the FOnAsyncRunCompleted delegate could be triggered from any thread.
-	 * - If SynchronousMode == ENeuralSynchronousMode::Synchronous, UNeuralNetwork::Run() will block the calling thread until completed, so a
-	 *   callback delegate is not required.
+	 * FOnAsyncRunCompleted is the delegate that gets triggered when an asynchronous inference has been completed (i.e., Run() was called with
+	 * NeuralSynchronousMode == ENeuralSynchronousMode::Asynchronous).
+	 *
+	 * GetOnAsyncRunCompletedDelegate() returns a reference of this class' FOnAsyncRunCompleted delegate, so the UNeuralNetwork's client can
+	 * subscribe to its callback function.
+	 *
+	 * Code example enabling it:
+	 * 	   AsynchronousNetwork->GetOnAsyncRunCompletedDelegate().BindUObject(this, &SomeClass::SomeFunctionToRunWhenAsyncRunFinished);
+	 *     AsynchronousNetwork->SetSynchronousMode(ENeuralSynchronousMode::Asynchronous);
+	 * Code example disabling it:
+	 *     AsynchronousNetwork->GetOnAsyncRunCompletedDelegate().Unbind();
+	 *     AsynchronousNetwork->SetSynchronousMode(ENeuralSynchronousMode::Synchronous);
+	 *
+	 * @see SynchronousMode(), GetSynchronousMode(), SetSynchronousMode(), GetThreadModeDelegateForAsyncRunCompleted(), SetThreadModeDelegateForAsyncRunCompleted() for more details.
 	 */
 	DECLARE_DELEGATE(FOnAsyncRunCompleted);
 	FOnAsyncRunCompleted& GetOnAsyncRunCompletedDelegate();
-	ENeuralThreadMode GetOnAsyncRunCompletedDelegateMode() const;
-	void SetOnAsyncRunCompletedDelegateMode(const ENeuralThreadMode InDelegateThreadMode);
+
+	/**
+	 * Getter and setter functions for ThreadModeDelegateForAsyncRunCompleted:
+	 * - GetThreadModeDelegateForAsyncRunCompleted(): It returns ThreadModeDelegateForAsyncRunCompleted.
+	 * - SetThreadModeDelegateForAsyncRunCompleted(): It sets ThreadModeDelegateForAsyncRunCompleted.
+	 * @see ThreadModeDelegateForAsyncRunCompleted and GetOnAsyncRunCompletedDelegate() for more details.
+	 */
+	ENeuralThreadMode GetThreadModeDelegateForAsyncRunCompleted() const;
+	void SetThreadModeDelegateForAsyncRunCompleted(const ENeuralThreadMode InThreadModeDelegateForAsyncRunCompleted);
 	
+// CONTINUE HERE
 	/**
-	 * Whether GPU execution is supported for this platform. It will return:
-	 * - True if DX12 is enabled, meaning UEAndORT can run on both the CPU and GPU. Also true if the current platform is not Windows. I will also
-	 *   return true if the back end is UEOnly.
-	 * - False if DX12 is disabled, meaning UEAndORT can run only run on the CPU. The user will need to enable DX12 to be able to run GPU, switch to
-	 *   CPU, or switch to the UEOnly back end.
+	 * Whether GPU execution is supported for this platform. It will return false for non-Windows platforms. On Windows:
+	 * - True if DX12 is enabled.
+	 * - False if DX12 is disabled. The user will need to enable DX12 to be able to run GPU on Windows or switch to CPU mode.
+	 * Note: I will also return true if the back end is UEOnly, but this is an unsupported and deprecated back end.
 	 */
 	bool IsGPUSupported() const;
 
 	/**
 	 * Functions to get/fill input.
-	 * These functions either take a TArray as input, return a modificable void* to fill the data, or return a constant FNeuralTensor(s) to see input
+	 * These functions either take a TArray as input, return a modifiable void* to fill the data, or return a constant FNeuralTensor(s) to see input
 	 * properties (e.g., size or dimensions).
 	 */
 	const FNeuralTensor& GetInputTensor(const int32 InTensorIndex = 0) const;
@@ -159,7 +200,7 @@ public:
 	 * Its output results can be retrieved with GetOutputTensor() or GetOutputTensors().
 	 *
 	 * If Run() is called asynchronously, this does not guarantee that calling SetInputFromArrayCopy multiple times will result in each one being
-	 * applied for a different Run. The user is responable of not calling SetInputFromArrayCopy until Run() is completed and its delegate
+	 * applied for a different Run. The user is responsible of not calling SetInputFromArrayCopy until Run() is completed and its delegate
 	 * (OnAsyncRunCompletedDelegate) called. Otherwise, the wrong results might be returned.
 	 */
 	void Run();
@@ -199,14 +240,17 @@ protected:
 	/**
 	 * SynchronousMode defines whether UNeuralNetwork::Run() will block the thread until completed (Synchronous), or whether it will run on a
 	 * background thread, not blocking the calling thread (Asynchronous).
-	 * If asynchronous, DelegateThreadMode will define whether the callback delegate is called from the game thread (highly recommended) or from
+	 * If asynchronous, ThreadModeDelegateForAsyncRunCompleted will define whether the callback delegate is called from the game thread (highly recommended) or from
 	 * any available thread (not fully thread safe).
+	 * - If ThreadModeDelegateForAsyncRunCompleted == ENeuralSynchronousMode::Asynchronous, the FOnAsyncRunCompleted delegate could be triggered from any thread.
+	 * - If ThreadModeDelegateForAsyncRunCompleted == ENeuralSynchronousMode::Synchronous, UNeuralNetwork::Run() will block the calling thread until completed, so a
+	 *   callback delegate is not required.
 	 * @see ENeuralSynchronousMode, ENeuralThreadMode for more details.
 	 */
 	UPROPERTY(Transient, VisibleAnywhere, Category = "Neural Network Inference")
 	ENeuralSynchronousMode SynchronousMode;
 	UPROPERTY(Transient, VisibleAnywhere, Category = "Neural Network Inference")
-	ENeuralThreadMode DelegateThreadMode;
+	ENeuralThreadMode ThreadModeDelegateForAsyncRunCompleted;
 
 	/**
 	 * Original model file path from which this UNeuralNetwork was loaded from.
@@ -244,7 +288,7 @@ private:
 
 	/**
 	 * Struct pointer containing the UE-and-ORT-based back end implementation.
-	 * PIMPL idiom to minimize memory when not using this back end and to to hide 3rd party dependencies.
+	 * PIMPL idiom to minimize memory when not using this back end and to hide 3rd party dependencies.
 	 * http://www.cppsamples.com/common-tasks/pimpl.html
 	 */
 	struct FImplBackEndUEAndORT;
@@ -315,8 +359,8 @@ public:
 	 * Internal and Editor-only functions not needed by the user.
 	 * Importing data and options used for loading the network.
 	 */
-	class UAssetImportData* GetAssetImportData() const;
-	class UAssetImportData* GetAndMaybeCreateAssetImportData();
+	TObjectPtr<class UAssetImportData> GetAssetImportData() const;
+	TObjectPtr<class UAssetImportData> GetAndMaybeCreateAssetImportData();
 #endif // WITH_EDITOR
 
 private:
@@ -337,7 +381,7 @@ private:
 #if WITH_EDITORONLY_DATA
 	/** Importing data and options used for loading the network. */
 	UPROPERTY(Instanced)
-	class UAssetImportData* AssetImportData;
+	TObjectPtr<class UAssetImportData> AssetImportData;
 #endif // WITH_EDITORONLY_DATA
 
 #if WITH_EDITOR
diff --git a/Engine/Plugins/Experimental/NeuralNetworkInference/Source/NeuralNetworkInference/Public/NeuralTensor.h b/Engine/Plugins/Experimental/NeuralNetworkInference/Source/NeuralNetworkInference/Public/NeuralTensor.h
--- a/Engine/Plugins/Experimental/NeuralNetworkInference/Source/NeuralNetworkInference/Public/NeuralTensor.h
+++ b/Engine/Plugins/Experimental/NeuralNetworkInference/Source/NeuralNetworkInference/Public/NeuralTensor.h
@@ -41,13 +41,13 @@ public:
 	 *  to 0 or a negative value if memory allocation is not required or the final size is unknown.
 	 * @param InName It sets Name. @see Name for more details.
 	 * @param InTensorType It sets TensorType. @see TensorType for more details.
-	 * @param InArray Alternative to InSizes/InVolume, and represents the input data to copy from. The templated constructor that uses InArray is a
+	 * @param InArray Alternative to InSizes/InVolume, and represents the input data to copy from. The template constructor that uses InArray is a
 	 *  simple (but not efficient) way of initializing a FNeuralTensor from the existing InArray TArray. It is equivalent to calling
 	 *  FNeuralTensor(..., InSizes/InVolume, ...) and then SetFromArrayCopy(InArray). It is not efficient because it makes a deep copy of the data
 	 *  of InArray. For maximum speed, rather than creating this intermediate InArray TArray, you can avoid the double copy by calling
 	 *  FNeuralTensor(..., InSizes/InVolume, ...) and then filling the tensor memory with GetData()/GetDataCasted().
 	 *
-	 * In addition, UnderlyingUInt8ArrayData is pre-allocated from InDataType and InSizes/InVolume. @see UnderlyingUInt8ArrayData for more details.
+	 * In addition, UnderlyingUInt8ArrayData is preallocated from InDataType and InSizes/InVolume. @see UnderlyingUInt8ArrayData for more details.
 	 */
 	FNeuralTensor(const ENeuralDataType InDataType = ENeuralDataType::None, const TArray<int64>& InSizes = TArray<int64>(),
 		const FString& InName = TEXT("FNeuralTensor"), const ENeuralTensorType InTensorType = ENeuralTensorType::Generic);
@@ -204,7 +204,7 @@ public:
 	 *   references/pointers will no longer be valid.
 	 * - These functions only modify the CPU memory, they do not synchronize CPU and GPU memory automatically. The user must do this accordingly.
 	 *
-	 * SetNumUninitialized resizes the tensor to the desired new size (i.e., modiying DataType, Sizes, Volume, UnderlyingUInt8ArrayData, etc.).
+	 * SetNumUninitialized resizes the tensor to the desired new size (i.e., modifying DataType, Sizes, Volume, UnderlyingUInt8ArrayData, etc.).
 	 * @param InDataType It sets DataType. If None, it will keep its current value. @see DataType for more details.
 	 * @param InSizes It sets Sizes and Volume. Set to empty (or omit this argument) if memory allocation is not required or the final tensor size is
 	 *  unknown. @see Sizes, Volume for more details.
@@ -362,7 +362,7 @@ protected:
 
 private:
 	/**
-	 * The array containing the underlying CPU data. It is pre-allocated based on the combination of DataType and Sizes/Volume (if Sizes/Volume are
+	 * The array containing the underlying CPU data. It is preallocated based on the combination of DataType and Sizes/Volume (if Sizes/Volume are
 	 * not empty/zero): UnderlyingUInt8ArrayData.Num() == sizeof(data type) x Volume.
 	 */
 	UPROPERTY()
@@ -382,7 +382,7 @@ private:
 	TSharedPtr<FRDGBufferUAVRef> BufferUAVRef;
 
 	/**
-	 * Auxiliary constructor for the templated constructor that takes a TArray<T>& InArray.
+	 * Auxiliary constructor for the template constructor that takes a TArray<T>& InArray.
 	 * InSizeOfT and sizeof(InDataType) should match, as well as the volume of InValueNum and InSizes.
 	 * It simply moves as much code as possible from the template constructor to the cpp file.
 	 */
@@ -390,7 +390,7 @@ private:
 		const TArray<int64>& InSizes, const FString& InName, const ENeuralTensorType InTensorType);
 
 	/**
-	 * Auxiliary function for the templated SetFromArrayCopy().
+	 * Auxiliary function for the template SetFromArrayCopy().
 	 * It simply moves as much code as possible from the template function to the cpp file.
 	 */
 	void SetFromPointer(const void* const InData, const int64 InSizeOfT, const int64 InDataSize);
diff --git a/Engine/Plugins/Experimental/NeuralNetworkInference/Source/NeuralNetworkInferenceEditor/Private/NeuralNetworkReimportFactory.cpp b/Engine/Plugins/Experimental/NeuralNetworkInference/Source/NeuralNetworkInferenceEditor/Private/NeuralNetworkReimportFactory.cpp
--- a/Engine/Plugins/Experimental/NeuralNetworkInference/Source/NeuralNetworkInferenceEditor/Private/NeuralNetworkReimportFactory.cpp
+++ b/Engine/Plugins/Experimental/NeuralNetworkInference/Source/NeuralNetworkInferenceEditor/Private/NeuralNetworkReimportFactory.cpp
@@ -45,7 +45,7 @@ bool UNeuralNetworkReimportFactory::CanReimport(UObject* Obj, TArray<FString>& O
 	UNeuralNetwork* Network = Cast<UNeuralNetwork>(Obj);
 	if (Network)
 	{
-		UAssetImportData* AssetImportData = Network->GetAssetImportData();
+		TObjectPtr<UAssetImportData> AssetImportData = Network->GetAssetImportData();
 		if (AssetImportData)
 		{
 			const bool bIsValidFile = UNeuralNetworkFactory::IsValidFile(AssetImportData->GetFirstFilename());
@@ -70,7 +70,7 @@ void UNeuralNetworkReimportFactory::SetReimportPaths(UObject* Obj, const TArray<
 	if (Network && ensure(NewReimportPaths.Num() == 1))
 	{
 		Network->Modify();
-		if (UAssetImportData* AssetImportData = Network->GetAssetImportData())
+		if (TObjectPtr<UAssetImportData> AssetImportData = Network->GetAssetImportData())
 		{
 			AssetImportData->UpdateFilenameOnly(NewReimportPaths[0]);
 		}
diff --git a/Engine/Plugins/Experimental/NeuralNetworkInference/Source/Private/NeuralNetworkInferenceQA/Private/ModelUnitTester.cpp b/Engine/Plugins/Experimental/NeuralNetworkInference/Source/Private/NeuralNetworkInferenceQA/Private/ModelUnitTester.cpp
--- a/Engine/Plugins/Experimental/NeuralNetworkInference/Source/Private/NeuralNetworkInferenceQA/Private/ModelUnitTester.cpp
+++ b/Engine/Plugins/Experimental/NeuralNetworkInference/Source/Private/NeuralNetworkInferenceQA/Private/ModelUnitTester.cpp
@@ -321,12 +321,12 @@ bool FModelUnitTester::ModelAccuracyTest(UNeuralNetwork* InOutNetwork, const ENe
 	const ENeuralDeviceType OriginalInputDeviceType = InOutNetwork->GetInputDeviceType();
 	const ENeuralDeviceType OriginalOutputDeviceType = InOutNetwork->GetOutputDeviceType();
 	const ENeuralSynchronousMode OriginalSynchronousMode = InOutNetwork->GetSynchronousMode();
-	const ENeuralThreadMode OriginalDelegateThreadMode = InOutNetwork->GetOnAsyncRunCompletedDelegateMode();
+	const ENeuralThreadMode OriginalThreadModeDelegateForAsyncRunCompleted = InOutNetwork->GetThreadModeDelegateForAsyncRunCompleted();
 	const UNeuralNetwork::ENeuralBackEnd OriginalBackEnd = InOutNetwork->GetBackEnd();
 	
 	// Set (a)synchronous Mode
 	InOutNetwork->SetSynchronousMode(InSynchronousMode);
-	InOutNetwork->SetOnAsyncRunCompletedDelegateMode(ENeuralThreadMode::AnyThread);
+	InOutNetwork->SetThreadModeDelegateForAsyncRunCompleted(ENeuralThreadMode::AnyThread);
 
 	// Set back end
 	if (!InOutNetwork->SetBackEnd(InBackEnd))
@@ -414,7 +414,7 @@ bool FModelUnitTester::ModelAccuracyTest(UNeuralNetwork* InOutNetwork, const ENe
 	// Reset to original network state
 	InOutNetwork->SetDeviceType(/*DeviceType*/OriginalDeviceType, /*InputDeviceType*/OriginalInputDeviceType, /*OutputDeviceType*/OriginalOutputDeviceType);
 	InOutNetwork->SetSynchronousMode(OriginalSynchronousMode);
-	InOutNetwork->SetOnAsyncRunCompletedDelegateMode(OriginalDelegateThreadMode);
+	InOutNetwork->SetThreadModeDelegateForAsyncRunCompleted(OriginalThreadModeDelegateForAsyncRunCompleted);
 	InOutNetwork->SetBackEnd(OriginalBackEnd);
 	// Test successful
 	return bDidGlobalTestPassed;
