pix3.h needs to be included between THIRD_PARTY_INCLUDES uses

diff --git a/Engine/Plugins/Experimental/NeuralNetworkInference/Source/NeuralNetworkInference/Private/NeuralNetworkImplBackEndUEAndORT.cpp b/Engine/Plugins/Experimental/NeuralNetworkInference/Source/NeuralNetworkInference/Private/NeuralNetworkImplBackEndUEAndORT.cpp
--- a/Engine/Plugins/Experimental/NeuralNetworkInference/Source/NeuralNetworkInference/Private/NeuralNetworkImplBackEndUEAndORT.cpp
+++ b/Engine/Plugins/Experimental/NeuralNetworkInference/Source/NeuralNetworkInference/Private/NeuralNetworkImplBackEndUEAndORT.cpp
@@ -50,7 +50,9 @@ NNI_THIRD_PARTY_INCLUDES_END
 #if defined(PLATFORM_WIN64)
 
 #if WITH_EDITOR && !UE_BUILD_SHIPPING
+NNI_THIRD_PARTY_INCLUDES_START
 	#include "pix3.h"
+NNI_THIRD_PARTY_INCLUDES_END
 	#define NNIGPUProfileMarker(Name) FNNIGPUProfiler::Instance()->Marker(Name)
 #else
 	#define NNIGPUProfileMarker(Name)
