Fixed usage of TCHAR_TO_ANSI in NeuralNetworkInference.

diff --git a/Engine/Plugins/Experimental/NeuralNetworkInference/Source/NeuralNetworkInference/Private/NeuralNetworkImplBackEndUEAndORT.cpp b/Engine/Plugins/Experimental/NeuralNetworkInference/Source/NeuralNetworkInference/Private/NeuralNetworkImplBackEndUEAndORT.cpp
--- a/Engine/Plugins/Experimental/NeuralNetworkInference/Source/NeuralNetworkInference/Private/NeuralNetworkImplBackEndUEAndORT.cpp
+++ b/Engine/Plugins/Experimental/NeuralNetworkInference/Source/NeuralNetworkInference/Private/NeuralNetworkImplBackEndUEAndORT.cpp
@@ -484,9 +484,9 @@ bool UNeuralNetwork::FImplBackEndUEAndORT::InitializedAndConfigureMembers(TShare
 		InOutImplBackEndUEAndORT = MakeShared<FImplBackEndUEAndORT>(InOutOnAsyncRunCompletedDelegate, InOutDelegateThreadMode, InOutResoucesCriticalSection);
 
 		// Set up ORT and create an environment
-		const char* const ModelFullFilePathCharPtr = TCHAR_TO_ANSI(*InModelFullFilePath);
-		// @todo: ModelFullFilePathCharPtr -> I thought any unique string would work, but it might be output logging file, so it has to be a non-existing file!
-		InOutImplBackEndUEAndORT->Environment = MakeUnique<Ort::Env>(ORT_LOGGING_LEVEL_WARNING, ModelFullFilePathCharPtr);
+		const auto ModelFullFilePathAnsiString = StringCast<ANSICHAR>(*InModelFullFilePath);
+		// @todo: ModelFullFilePathAnsiString -> I thought any unique string would work, but it might be output logging file, so it has to be a non-existing file!
+		InOutImplBackEndUEAndORT->Environment = MakeUnique<Ort::Env>(ORT_LOGGING_LEVEL_WARNING, ModelFullFilePathAnsiString.Get());
 
 		InOutImplBackEndUEAndORT->Allocator = MakeUnique<Ort::AllocatorWithDefaultOptions>();
 
diff --git a/Engine/Plugins/Experimental/NeuralNetworkInference/Source/Private/ModelProto/Private/ModelProtoStringParser.cpp b/Engine/Plugins/Experimental/NeuralNetworkInference/Source/Private/ModelProto/Private/ModelProtoStringParser.cpp
--- a/Engine/Plugins/Experimental/NeuralNetworkInference/Source/Private/ModelProto/Private/ModelProtoStringParser.cpp
+++ b/Engine/Plugins/Experimental/NeuralNetworkInference/Source/Private/ModelProto/Private/ModelProtoStringParser.cpp
@@ -148,11 +148,11 @@ TArray<uint8> FModelProtoStringParser::GetModelProtoCharAsUInt8Array(const TMap<
 	const FString RawDataString = FModelProtoStringParser::GetModelProtoStringOrEmpty(InProtoMap, InKey);
 	if (!RawDataString.IsEmpty())
 	{
-		const char* const RawDataCharPtr = TCHAR_TO_ANSI(*RawDataString);
+		const auto RawDataAnsiString = StringCast<ANSICHAR>(*RawDataString);
 		// Allocate memory
 		ProtoCharAsUint8Array.SetNumUninitialized(RawDataString.Len());
 		// Copy memory
-		FMemory::Memcpy(ProtoCharAsUint8Array.GetData(), RawDataCharPtr, ProtoCharAsUint8Array.Num());
+		FMemory::Memcpy(ProtoCharAsUint8Array.GetData(), RawDataAnsiString.Get(), ProtoCharAsUint8Array.Num());
 	}
 	// Return array
 	return ProtoCharAsUint8Array;
