Fixed doc typo for inline comments

#rnx

diff --git a/Engine/Plugins/Experimental/NeuralNetworkInference/Source/NeuralNetworkInference/Internal/NeuralOperator.h b/Engine/Plugins/Experimental/NeuralNetworkInference/Source/NeuralNetworkInference/Internal/NeuralOperator.h
--- a/Engine/Plugins/Experimental/NeuralNetworkInference/Source/NeuralNetworkInference/Internal/NeuralOperator.h
+++ b/Engine/Plugins/Experimental/NeuralNetworkInference/Source/NeuralNetworkInference/Internal/NeuralOperator.h
@@ -99,7 +99,8 @@ protected:
 	const int32 Version;
 	const TSet<uint32> PotentiallyInlinedTensors;
 	int32 InlinedTensor;
-	TArray<FNeuralTensor*> AuxiliaryTensors; /* Only for some operators that might need them (e.g., FConvOperator and FConvTransposeOperator). */
+	/** Only for some operators that might need them (e.g., FConvOperator and FConvTransposeOperator). */
+	TArray<FNeuralTensor*> AuxiliaryTensors;
 
 	/**
 	 * 2 cases:
@@ -147,8 +148,10 @@ protected:
 	virtual const TArray<FNeuralTensor*>& GetOutputTensorsConst() const final;
 
 private:
-	TArray<FNeuralTensor*> PrivateInputTensors; /* Const for non-inlined operators. */
-	TArray<FNeuralTensor*> PrivateOutputTensors; /* Only for non-inlined operators. */
+	/** Const for non-inlined operators. */
+	TArray<FNeuralTensor*> PrivateInputTensors;
+	/** Only for non-inlined operators. */
+	TArray<FNeuralTensor*> PrivateOutputTensors;
 };
 
 
diff --git a/Engine/Plugins/Experimental/NeuralNetworkInference/Source/NeuralNetworkInference/Internal/NeuralOperators/ConvBaseOperator.h b/Engine/Plugins/Experimental/NeuralNetworkInference/Source/NeuralNetworkInference/Internal/NeuralOperators/ConvBaseOperator.h
--- a/Engine/Plugins/Experimental/NeuralNetworkInference/Source/NeuralNetworkInference/Internal/NeuralOperators/ConvBaseOperator.h
+++ b/Engine/Plugins/Experimental/NeuralNetworkInference/Source/NeuralNetworkInference/Internal/NeuralOperators/ConvBaseOperator.h
@@ -11,10 +11,14 @@ class NEURALNETWORKINFERENCE_API FConvBaseOperator : public FNeuralOperator
 public:
 	enum class EAutoPad : uint8
 	{
-		NotSet, /* Explicit padding is used */
-		SameUpper, /* Auto-padding the input so that input and output match in size (if odd number, extra padding at the beginning) */
-		SameLower, /* Auto-padding the input so that input and output match in size (if odd number, extra padding at the end) */
-		Valid /* No padding is used */
+		/** Explicit padding is used */
+		NotSet,
+		/** Auto-padding the input so that input and output match in size (if odd number, extra padding at the beginning) */
+		SameUpper,
+		/** Auto-padding the input so that input and output match in size (if odd number, extra padding at the end) */
+		SameLower,
+		/** No padding is used */
+		Valid
 	};
 
 	virtual ~FConvBaseOperator();
diff --git a/Engine/Plugins/Experimental/NeuralNetworkInference/Source/NeuralNetworkInference/Private/NeuralNetworkImplBackEndUEAndORT.h b/Engine/Plugins/Experimental/NeuralNetworkInference/Source/NeuralNetworkInference/Private/NeuralNetworkImplBackEndUEAndORT.h
--- a/Engine/Plugins/Experimental/NeuralNetworkInference/Source/NeuralNetworkInference/Private/NeuralNetworkImplBackEndUEAndORT.h
+++ b/Engine/Plugins/Experimental/NeuralNetworkInference/Source/NeuralNetworkInference/Private/NeuralNetworkImplBackEndUEAndORT.h
@@ -69,17 +69,23 @@ private:
 	TUniquePtr<Ort::Session> Session;
 	TUniquePtr<Ort::AllocatorWithDefaultOptions> Allocator;
 	TUniquePtr<Ort::SessionOptions> SessionOptions;
-	/** Tensor-related variables */
-	TUniquePtr<Ort::MemoryInfo> AllocatorInfo; /* Memory allocator information */
+	/** Tensor-related variable: Memory allocator information */
+	TUniquePtr<Ort::MemoryInfo> AllocatorInfo;
 #ifdef PLATFORM_WIN64
 	const OrtDmlApi* DmlApi;
-	TUniquePtr<Ort::MemoryInfo> DmlGPUMemoryInfo; /* DirectML GPU memory information */
-	TArray<void*> DmlGPUResources; /* Shared D3D12 resources with DirectML GPU execution provider */
+	/** DirectML GPU memory information */
+	TUniquePtr<Ort::MemoryInfo> DmlGPUMemoryInfo;
+	/** Shared D3D12 resources with DirectML GPU execution provider */
+	TArray<void*> DmlGPUResources;
 #endif
-	TArray<Ort::Value> InputOrtTensors; /* Actual ONNXRuntime tensors */
-	TArray<const char*> InputTensorNames; /* Tensor names */
-	TArray<Ort::Value> OutputOrtTensors; /* Actual ONNXRuntime tensors */
-	TArray<const char*> OutputTensorNames; /* Tensor names */
+	/** Actual ONNXRuntime tensors */
+	TArray<Ort::Value> InputOrtTensors;
+	/** Tensor names */
+	TArray<const char*> InputTensorNames;
+	/** Actual ONNXRuntime tensors */
+	TArray<Ort::Value> OutputOrtTensors;
+	/** Tensor names */
+	TArray<const char*> OutputTensorNames;
 	
 	// Helper class to run session as an async task
 	class FNeuralNetworkAsyncTask : public FNonAbandonableTask
diff --git a/Engine/Plugins/Experimental/NeuralNetworkInference/Source/NeuralNetworkInference/Public/NeuralEnumClasses.h b/Engine/Plugins/Experimental/NeuralNetworkInference/Source/NeuralNetworkInference/Public/NeuralEnumClasses.h
--- a/Engine/Plugins/Experimental/NeuralNetworkInference/Source/NeuralNetworkInference/Public/NeuralEnumClasses.h
+++ b/Engine/Plugins/Experimental/NeuralNetworkInference/Source/NeuralNetworkInference/Public/NeuralEnumClasses.h
@@ -13,16 +13,26 @@
 UENUM()
 enum class ENeuralDataType : uint8
 {
-	Float, /* 32-bit floating number. */
-	//Double, /* 64-bit floating number. */
-	//Int8, /* 8-bit signed integer. */
-	//Int16, /* 16-bit signed integer. */
-	Int32, /* 32-bit signed integer. */
-	Int64, /* 64-bit signed integer. */
-	//UInt8, /* 8-bit unsigned integer. */
-	UInt32, /* 32-bit unsigned integer. */
-	UInt64, /* 64-bit unsigned integer. */
-	None /* To be used in special cases, e.g., if the data type is unknown yet. */
+	/** 32-bit floating number. */
+	Float,
+	/** 64-bit floating number. */
+	//Double,
+	/** 8-bit signed integer. */
+	//Int8,
+	/** 16-bit signed integer. */
+	//Int16,
+	/** 32-bit signed integer. */
+	Int32,
+	/** 64-bit signed integer. */
+	Int64,
+	/** 8-bit unsigned integer. */
+	//UInt8,
+	/** 32-bit unsigned integer. */
+	UInt32,
+	/** 64-bit unsigned integer. */
+	UInt64,
+	/** To be used in special cases, e.g., if the data type is unknown yet. */
+	None
 };
 
 /**
@@ -31,8 +41,10 @@ enum class ENeuralDataType : uint8
 UENUM()
 enum class ENeuralDeviceType : uint8
 {
-	CPU, /* The operation will occur on the CPU. */
-	GPU /* The operation will occur on the GPU. */
+	/** The operation will occur on the CPU. */
+	CPU,
+	/** The operation will occur on the GPU. */
+	GPU
 };
 
 /**
diff --git a/Engine/Plugins/Experimental/NeuralNetworkInference/Source/Private/ThirdPartyHelperAndDLLLoader/Internal/ThirdPartyWarningDisabler.h b/Engine/Plugins/Experimental/NeuralNetworkInference/Source/Private/ThirdPartyHelperAndDLLLoader/Internal/ThirdPartyWarningDisabler.h
--- a/Engine/Plugins/Experimental/NeuralNetworkInference/Source/Private/ThirdPartyHelperAndDLLLoader/Internal/ThirdPartyWarningDisabler.h
+++ b/Engine/Plugins/Experimental/NeuralNetworkInference/Source/Private/ThirdPartyHelperAndDLLLoader/Internal/ThirdPartyWarningDisabler.h
@@ -24,7 +24,7 @@ Many third party headers require some care when importing. NNI third party inclu
 	__pragma(warning(disable: 6387)) /* C6387: 'X' could be '0':  this does not adhere to the specification for the function 'memset'. */ \
 	__pragma(warning(disable: 6388)) /* C6388: '*X' might not be '0':  this does not adhere to the specification for the function 'Y'. */ \
 	__pragma(warning(disable: 28020)) /* C28020: The expression '0<=_Param_(1)&&_Param_(1)<=400-1' is not true at this call. */ \
-	__pragma(warning(disable: 28196)) /* C28196: his warning indicates that the function prototype for the function being analyzed has a __notnull, __null or __drv_valueIs on an _Out_ parameter or the return value, but the value returned is inconsistent with that annotation. */ \
+	__pragma(warning(disable: 28196)) /* C28196: This warning indicates that the function prototype for the function being analyzed has a __notnull, __null or __drv_valueIs on an _Out_ parameter or the return value, but the value returned is inconsistent with that annotation. */ \
 	UE_PUSH_MACRO("check") \
 	UE_PUSH_MACRO("TEXT")
 #else
