1. UNeuralNetwork UI editable! (but causing a slow Load() with EVERY change).
2. UNN more documented.

#rnx

diff --git a/Engine/Plugins/Experimental/NeuralNetworkInference/Source/NeuralNetworkInference/Private/NeuralNetwork.cpp b/Engine/Plugins/Experimental/NeuralNetworkInference/Source/NeuralNetworkInference/Private/NeuralNetwork.cpp
--- a/Engine/Plugins/Experimental/NeuralNetworkInference/Source/NeuralNetworkInference/Private/NeuralNetwork.cpp
+++ b/Engine/Plugins/Experimental/NeuralNetworkInference/Source/NeuralNetworkInference/Private/NeuralNetwork.cpp
@@ -155,6 +155,25 @@ void UNeuralNetwork::SetDeviceType(const ENeuralDeviceType InDeviceType, const E
 	}
 }
 
+bool UNeuralNetwork::IsGPUSupported() const
+{
+	// UEAndORT
+	if (BackEndForCurrentPlatform == ENeuralBackEnd::UEAndORT)
+	{
+		return FImplBackEndUEAndORT::IsGPUSupported();
+	}
+
+	// UEOnly
+	else if (BackEndForCurrentPlatform == ENeuralBackEnd::UEOnly)
+	{
+		return true;
+	}
+
+	// Unknown
+	UE_LOG(LogNeuralNetworkInference, Warning, TEXT("UNeuralNetwork::IsGPUSupported(): Unknown [BackEnd,BackEndForCurrentPlatform] = [%d,%d]."), (int32)BackEnd, (int32)BackEndForCurrentPlatform);
+	return false;
+}
+
 ENeuralSynchronousMode UNeuralNetwork::GetSynchronousMode() const
 {
 	return SynchronousMode;
@@ -189,25 +208,6 @@ void UNeuralNetwork::ResetStats()
 	InputMemoryTransferStatsModule.ResetStats();
 }
 
-bool UNeuralNetwork::IsGPUSupported() const
-{
-	// UEAndORT
-	if (BackEndForCurrentPlatform == ENeuralBackEnd::UEAndORT)
-	{
-		return FImplBackEndUEAndORT::IsGPUSupported();
-	}
-
-	// UEOnly
-	else if (BackEndForCurrentPlatform == ENeuralBackEnd::UEOnly)
-	{
-		return true;
-	}
-
-	// Unknown
-	UE_LOG(LogNeuralNetworkInference, Warning, TEXT("UNeuralNetwork::IsGPUSupported(): Unknown [BackEnd,BackEndForCurrentPlatform] = [%d,%d]."), (int32)BackEnd, (int32)BackEndForCurrentPlatform);
-	return false;
-}
-
 #define GET_TENSOR_CODE(InTensorIndex, UEAndORTOnly_GetIndexes, UEOnly_GetTensorsFunction, UEOnly_Tensors) \
 	/* Sanity check */ \
 	checkf(bIsLoaded, TEXT("Call UNeuralNetwork::Load() to load a model first.")); \
@@ -637,11 +637,18 @@ void UNeuralNetwork::ReimportAssetFromEditorData()
 		// Ensure that the file provided by the path exists
 		if (IFileManager::Get().FileSize(*ImportedFilename) != INDEX_NONE)
 		{
-			UE_LOG(LogNeuralNetworkInference, Display, TEXT("Performing atomic reimport of [%s]"), *ImportedFilename);
+			UE_LOG(LogNeuralNetworkInference, Display, TEXT("Performing atomic reimport of [%s]..."), *ImportedFilename);
 			Load(ImportedFilename);
 		}
 	}
 }
+
+void UNeuralNetwork::PostEditChangeProperty(FPropertyChangedEvent& PropertyChangedEvent)
+{
+	Super::PostEditChangeProperty(PropertyChangedEvent);
+	UE_LOG(LogNeuralNetworkInference, Display, TEXT("UNeuralNetwork property changed, reloading network..."));
+	Load();
+}
 #endif // WITH_EDITOR
 
 void UNeuralNetwork::PostInitProperties()
diff --git a/Engine/Plugins/Experimental/NeuralNetworkInference/Source/NeuralNetworkInference/Public/NeuralNetwork.h b/Engine/Plugins/Experimental/NeuralNetworkInference/Source/NeuralNetworkInference/Public/NeuralNetwork.h
--- a/Engine/Plugins/Experimental/NeuralNetworkInference/Source/NeuralNetworkInference/Public/NeuralNetwork.h
+++ b/Engine/Plugins/Experimental/NeuralNetworkInference/Source/NeuralNetworkInference/Public/NeuralNetwork.h
@@ -28,7 +28,12 @@
  *		const FString ONNXModelFilePath = TEXT("SOME_PARENT_FOLDER/SOME_ONNX_FILE_NAME.onnx");
  *		if (Network->Load(ONNXModelFilePath))
  *		{
- *			Network->SetDeviceType(ENeuralDeviceType::CPU); // Set to CPU/GPU mode
+ *			// Pick between option a or b
+ *			// Option a) Set to GPU
+ *			if (Network->IsGPUSupported())
+ *				Network->SetDeviceType(ENeuralDeviceType::GPU);
+ *			// Option b) Set to CPU
+ *			Network->SetDeviceType(ENeuralDeviceType::CPU);
  *		}
  *
  * 1.b. Loading a UNeuralNetwork from a previously-created UNeuralNetwork UAsset:
@@ -37,7 +42,7 @@
  *		UNeuralNetwork* Network = LoadObject<UNeuralNetwork>((UObject*)GetTransientPackage(), *NetworkUAssetFilePath);
  *		// Check that the network was successfully loaded
  *		check(Network->IsLoaded());
- *      // Optionally set to CPU/GPU mode. This step is optional because otherwise it will use the device type that was saved on the loaded UAsset
+ *      // Optionally set to CPU/GPU mode. This step is optional, if not called, it will use the device type that was saved on the loaded UAsset
  *      // Network->SetDeviceType(ENeuralDeviceType::CPU);
  *
  * 2.a. Running inference:
@@ -106,14 +111,28 @@ public:
 	 * - GetDeviceType(): It returns DeviceType.
 	 * - GetInputDeviceType(): It returns InputDeviceType.
 	 * - GetOutputDeviceType(): It returns OutputDeviceType.
-	 * - SetDeviceType(): It sets DeviceType, InputDeviceType, and OutputDeviceType.
-	 * @see DeviceType, InputDeviceType, and OutputDeviceType for more details.
+	 * - SetDeviceType(): It sets DeviceType, InputDeviceType, and OutputDeviceType. If you are setting it to GPU, check IsGPUSupported() first.
+	 * @see DeviceType, InputDeviceType, OutputDeviceType, and IsGPUSupported() for more details.
 	 */
 	ENeuralDeviceType GetDeviceType() const;
 	ENeuralDeviceType GetInputDeviceType() const;
 	ENeuralDeviceType GetOutputDeviceType() const;
 	void SetDeviceType(const ENeuralDeviceType InDeviceType, const ENeuralDeviceType InInputDeviceType = ENeuralDeviceType::CPU,
 		const ENeuralDeviceType InOutputDeviceType = ENeuralDeviceType::CPU);
+	
+	/**
+	 * Whether GPU execution (i.e., SetDeviceType(ENeuralDeviceType::GPU)) is supported for this platform:
+	 * - On Windows:
+	 *     - True if DX12 is enabled.
+	 *     - False if DX12 is disabled. The user will need to enable DX12 on Unreal Engine to be able to run GPU on Windows or switch to CPU mode.
+	 * - Non-Windows platforms: False for now (GPU not supported on non-Windows devices yet).
+	 * Deprecated: I will also return true if the back end is UEOnly, but this is an unsupported and deprecated back end that should not be used.
+	 * Code sample:
+	 *		if (Network->IsGPUSupported())
+	 *			Network->SetDeviceType(ENeuralDeviceType::GPU);
+	 * @see DeviceType and SetDeviceType() for more details.
+	 */
+	bool IsGPUSupported() const;
 
 	/**
 	 * Getter and setter functions for SynchronousMode:
@@ -151,16 +170,8 @@ public:
 	 */
 	ENeuralThreadMode GetThreadModeDelegateForAsyncRunCompleted() const;
 	void SetThreadModeDelegateForAsyncRunCompleted(const ENeuralThreadMode InThreadModeDelegateForAsyncRunCompleted);
-	
-// CONTINUE HERE
-	/**
-	 * Whether GPU execution is supported for this platform. It will return false for non-Windows platforms. On Windows:
-	 * - True if DX12 is enabled.
-	 * - False if DX12 is disabled. The user will need to enable DX12 to be able to run GPU on Windows or switch to CPU mode.
-	 * Note: I will also return true if the back end is UEOnly, but this is an unsupported and deprecated back end.
-	 */
-	bool IsGPUSupported() const;
 
+// CONTINUE HERE
 	/**
 	 * Functions to get/fill input.
 	 * These functions either take a TArray as input, return a modifiable void* to fill the data, or return a constant FNeuralTensor(s) to see input
@@ -222,7 +233,7 @@ protected:
 	 * Whether Run() will use CPU or GPU acceleration hardware.
 	 * If SetDeviceType() is never called, the default device (EDeviceType::CPU) will be used.
 	 */
-	UPROPERTY(VisibleAnywhere, Category = "Neural Network Inference")
+	UPROPERTY(EditAnywhere, Category = "Neural Network Inference")
 	ENeuralDeviceType DeviceType;
 	
 	/**
@@ -231,10 +242,10 @@ protected:
 	 *  - InputDeviceType: Whether Run() will expect the input data in CPU (Run will upload the memory to the GPU) or GPU (no upload needed).
 	 *  - OutputDeviceType: Whether Run() will return output data in CPU (Run will download the memory to the CPU) or GPU (no download needed).
 	 */
-	UPROPERTY(VisibleAnywhere, Category = "Neural Network Inference")
+	UPROPERTY(EditAnywhere, Category = "Neural Network Inference")
 	ENeuralDeviceType InputDeviceType;
 	
-	UPROPERTY(VisibleAnywhere, Category = "Neural Network Inference")
+	UPROPERTY(EditAnywhere, Category = "Neural Network Inference")
 	ENeuralDeviceType OutputDeviceType;
 	
 	/**
@@ -247,9 +258,9 @@ protected:
 	 *   callback delegate is not required.
 	 * @see ENeuralSynchronousMode, ENeuralThreadMode for more details.
 	 */
-	UPROPERTY(Transient, VisibleAnywhere, Category = "Neural Network Inference")
+	UPROPERTY(Transient, EditAnywhere, Category = "Neural Network Inference")
 	ENeuralSynchronousMode SynchronousMode;
-	UPROPERTY(Transient, VisibleAnywhere, Category = "Neural Network Inference")
+	UPROPERTY(Transient, EditAnywhere, Category = "Neural Network Inference")
 	ENeuralThreadMode ThreadModeDelegateForAsyncRunCompleted;
 
 	/**
@@ -335,8 +346,10 @@ public:
 		 * macro that checks whether UEAndORT support exists for the current platform.
 		 */
 		UEAndORT,
-		UEOnly, /* It might be slower than the UEAndORT back end, but it will compile in all platforms and OSs compatible with Unreal Engine. */
-		Auto /* Recommended value. It will use the efficient UEAndORT if supported by the platform, and fall back to UEOnly otherwise. */
+		/** It might be slower than the UEAndORT back end, but it will compile in all platforms and OSs compatible with Unreal Engine. */
+		UEOnly,
+		/** Recommended value. It will use the efficient UEAndORT if supported by the platform, and fall back to UEOnly otherwise. */
+		Auto
 	};
 
 	/**
@@ -387,6 +400,7 @@ private:
 #if WITH_EDITOR
 	/** Editor-only function: Re-import asset with editor data (imported file). */
 	void ReimportAssetFromEditorData();
+	virtual void PostEditChangeProperty(FPropertyChangedEvent& PropertyChangedEvent) override;
 #endif // WITH_EDITOR
 
 	//~UObject interface
