DirectML DLL loader checks file exists on hard disk and ensureMsgf.
#rnx

diff --git a/Engine/Plugins/Experimental/NeuralNetworkInference/Source/Private/ThirdPartyHelperAndDLLLoader/Private/Module.cpp b/Engine/Plugins/Experimental/NeuralNetworkInference/Source/Private/ThirdPartyHelperAndDLLLoader/Private/Module.cpp
--- a/Engine/Plugins/Experimental/NeuralNetworkInference/Source/Private/ThirdPartyHelperAndDLLLoader/Private/Module.cpp
+++ b/Engine/Plugins/Experimental/NeuralNetworkInference/Source/Private/ThirdPartyHelperAndDLLLoader/Private/Module.cpp
@@ -3,6 +3,8 @@
 #include "Module.h"
 #include "HAL/PlatformProcess.h"
 #include "Interfaces/IPluginManager.h"
+#include "Misc/Paths.h"
+#include "ThirdPartyHelperAndDLLLoaderUtils.h"
 
 // This code will execute after your module is loaded into memory; the exact timing is specified in the .uplugin file per-module
 void FThirdPartyHelperAndDLLLoaderModule::StartupModule()
@@ -12,6 +14,14 @@ void FThirdPartyHelperAndDLLLoaderModule::StartupModule()
 	const FString DirectMLRuntimeBinPath = PluginDir / TEXT(PREPROCESSOR_TO_STRING(DIRECTML_PLATFORM_PATH));
 	const FString DirectMLDLLPath = DirectMLRuntimeBinPath / TEXT(PREPROCESSOR_TO_STRING(DIRECTML_DLL_NAME));
 
+	// Sanity check
+	if (!FPaths::FileExists(DirectMLDLLPath))
+	{
+		const FString ErrorMessage = FString::Format(TEXT("DirectML DLL file not found in \"{0}\"."),
+			{ IFileManager::Get().ConvertToAbsolutePathForExternalAppForRead(*DirectMLDLLPath) });
+		UE_LOG(LogNeuralNetworkInferenceThirdPartyHelperAndDLLLoader, Warning, TEXT("FThirdPartyHelperAndDLLLoaderModule::StartupModule(): %s"), *ErrorMessage);
+		checkf(false, TEXT("%s"), *ErrorMessage);
+	}
 	FPlatformProcess::PushDllDirectory(*DirectMLRuntimeBinPath);
 	DirectMLDLLHandle = FPlatformProcess::GetDllHandle(*DirectMLDLLPath);
 	FPlatformProcess::PopDllDirectory(*DirectMLRuntimeBinPath);
