1. Added SetFromVoidPointerCopy to both UNN and UNTensor.
2. FNeuralTensor::SetFromPointer renamed as FNeuralTensor::SetFromVoidPointer to match its public equivalent.
3. Removed confusing UNN::CreateInputArrayCopy/OUtput and moved to QA instead (as they onlyh copy the TArray, not the other properties).
4. More doc added to UNN.

#rnx

diff --git a/Engine/Plugins/Experimental/NeuralNetworkInference/Source/NeuralNetworkInference/Private/NeuralNetwork.cpp b/Engine/Plugins/Experimental/NeuralNetworkInference/Source/NeuralNetworkInference/Private/NeuralNetwork.cpp
--- a/Engine/Plugins/Experimental/NeuralNetworkInference/Source/NeuralNetworkInference/Private/NeuralNetwork.cpp
+++ b/Engine/Plugins/Experimental/NeuralNetworkInference/Source/NeuralNetworkInference/Private/NeuralNetwork.cpp
@@ -87,7 +87,8 @@ bool UNeuralNetwork::Load(const FString& InModelFilePath)
 			return false;
 		}
 		// Read file into ModelReadFromFileInBytes
-		// Source: [REDACTED_PROFILE_URL]/onnxruntime/blob/894fc828587c919d815918c4da6cde314e5d54ed/onnxruntime/test/shared_lib/test_model_loading.cc#L21-L31
+		// Source:
+		// github.com/microsoft/onnxruntime/blob/894fc828587c919d815918c4da6cde314e5d54ed/onnxruntime/test/shared_lib/test_model_loading.cc#L21-L31
 		if (!FFileHelper::LoadFileToArray(ModelReadFromFileInBytes, *ModelFullFilePath))
 		{
 			UE_LOG(LogNeuralNetworkInference, Warning, TEXT("UNeuralNetwork::Load(): Error reading model \"%s\"."), *ModelFullFilePath);
@@ -140,7 +141,8 @@ ENeuralDeviceType UNeuralNetwork::GetOutputDeviceType() const
 	return OutputDeviceType;
 }
 
-void UNeuralNetwork::SetDeviceType(const ENeuralDeviceType InDeviceType, const ENeuralDeviceType InInputDeviceType, const ENeuralDeviceType InOutputDeviceType)
+void UNeuralNetwork::SetDeviceType(const ENeuralDeviceType InDeviceType, const ENeuralDeviceType InInputDeviceType,
+	const ENeuralDeviceType InOutputDeviceType)
 {
 	const FScopeLock ResourcesLock(&ResoucesCriticalSection);
 	if (DeviceType != InDeviceType || InputDeviceType != InInputDeviceType || OutputDeviceType != InOutputDeviceType)
@@ -170,7 +172,8 @@ bool UNeuralNetwork::IsGPUSupported() const
 	}
 
 	// Unknown
-	UE_LOG(LogNeuralNetworkInference, Warning, TEXT("UNeuralNetwork::IsGPUSupported(): Unknown [BackEnd,BackEndForCurrentPlatform] = [%d,%d]."), (int32)BackEnd, (int32)BackEndForCurrentPlatform);
+	UE_LOG(LogNeuralNetworkInference, Warning, TEXT("UNeuralNetwork::IsGPUSupported(): Unknown [BackEnd,BackEndForCurrentPlatform] = [%d,%d]."),
+		(int32)BackEnd, (int32)BackEndForCurrentPlatform);
 	return false;
 }
 
@@ -202,10 +205,61 @@ void UNeuralNetwork::SetThreadModeDelegateForAsyncRunCompleted(const ENeuralThre
 	ThreadModeDelegateForAsyncRunCompleted = InThreadModeDelegateForAsyncRunCompleted;
 }
 
-void UNeuralNetwork::ResetStats()
+int64 UNeuralNetwork::GetInputTensorNumber() const
 {
-	ComputeStatsModule.ResetStats();
-	InputMemoryTransferStatsModule.ResetStats();
+	// Sanity check
+	if (!bIsLoaded)
+	{
+		UE_LOG(LogNeuralNetworkInference, Warning,
+			TEXT("UNeuralNetwork::GetInputTensorNumber(): Call UNeuralNetwork::Load() to load a model first."));
+		return -1;
+	}
+
+	// UEAndORT
+	if (BackEndForCurrentPlatform == ENeuralBackEnd::UEAndORT)
+	{
+		return ImplBackEndUEAndORT->InputTensors.Num();
+	}
+
+	// UEOnly
+	else if (BackEndForCurrentPlatform == ENeuralBackEnd::UEOnly)
+	{
+		return ImplBackEndUEOnly->TensorManager.GetInputIndexes().Num();
+	}
+
+	// Unknown
+	checkf(false, TEXT("UNeuralNetwork::GetInputTensorNumber(): Unknown [BackEnd,BackEndForCurrentPlatform] = [%d,%d]."), (int32)BackEnd,
+		(int32)BackEndForCurrentPlatform);
+	return -1;
+}
+
+int64 UNeuralNetwork::GetOutputTensorNumber() const
+{
+	// Sanity check
+	if (!bIsLoaded)
+	{
+		UE_LOG(LogNeuralNetworkInference, Warning,
+			TEXT("UNeuralNetwork::GetOutputTensorNumber(): Call UNeuralNetwork::Load() to load a model first."));
+		return -1;
+	}
+
+	// UEAndORT
+	if (BackEndForCurrentPlatform == ENeuralBackEnd::UEAndORT)
+	{
+		return ImplBackEndUEAndORT->OutputTensors.Num();
+	}
+
+	// UEOnly
+	else if (BackEndForCurrentPlatform == ENeuralBackEnd::UEOnly)
+	{
+		return ImplBackEndUEOnly->TensorManager.GetOutputIndexes().Num();
+	}
+
+	// Unknown
+	UE_LOG(LogNeuralNetworkInference, Warning,
+		TEXT("UNeuralNetwork::GetOutputTensorNumber(): Unknown [BackEnd,BackEndForCurrentPlatform] = [%d,%d]."),
+		(int32)BackEnd, (int32)BackEndForCurrentPlatform);
+	return -1;
 }
 
 #define GET_TENSOR_CODE(InTensorIndex, UEAndORTOnly_GetIndexes, UEOnly_GetTensorsFunction, UEOnly_Tensors) \
@@ -231,15 +285,20 @@ const FNeuralTensor& UNeuralNetwork::GetInputTensor(const int32 InTensorIndex) c
 	GET_TENSOR_CODE(InTensorIndex, GetInputIndexes, GetTensors, InputTensors);
 }
 
+const FNeuralTensor& UNeuralNetwork::GetOutputTensor(const int32 InTensorIndex) const
+{
+	GET_TENSOR_CODE(InTensorIndex, GetOutputIndexes, GetTensors, OutputTensors);
+}
+
 void UNeuralNetwork::SetInputFromArrayCopy(const TArray<float>& InArray, const int32 InTensorIndex)
 {
 	const FScopeLock ResourcesLock(&ResoucesCriticalSection);
 	// Sanity check
 	if (!bIsLoaded)
 	{
-		UE_LOG(LogNeuralNetworkInference, Warning, TEXT("UNeuralNetwork::SetInputFromArrayCopy(): Call UNeuralNetwork::Load() to load a model first."));
+		UE_LOG(LogNeuralNetworkInference, Warning,
+			TEXT("UNeuralNetwork::SetInputFromArrayCopy(): Call UNeuralNetwork::Load() to load a model first."));
 	}
-	
 	FNeuralTimer RunTimer;
 	RunTimer.Tic();
 	// UEAndORT
@@ -256,135 +315,76 @@ void UNeuralNetwork::SetInputFromArrayCopy(const TArray<float>& InArray, const i
 	// Unknown
 	else
 	{
-		UE_LOG(LogNeuralNetworkInference, Warning, TEXT("UNeuralNetwork::SetInputFromArrayCopy(): Unknown [BackEnd,BackEndForCurrentPlatform] = [%d,%d]."), (int32)BackEnd, (int32)BackEndForCurrentPlatform);
+		UE_LOG(LogNeuralNetworkInference, Warning,
+			TEXT("UNeuralNetwork::SetInputFromArrayCopy(): Unknown [BackEnd,BackEndForCurrentPlatform] = [%d,%d]."),
+			(int32)BackEnd, (int32)BackEndForCurrentPlatform);
 	}
 	InputMemoryTransferStatsModule.StoreSample(RunTimer.Toc());
 }
 
-void* UNeuralNetwork::GetInputDataPointerMutable(const int32 InTensorIndex)
+void UNeuralNetwork::SetInputFromVoidPointerCopy(const void* const InVoidPtr, const int32 InTensorIndex)
 {
 	const FScopeLock ResourcesLock(&ResoucesCriticalSection);
-
 	// Sanity check
 	if (!bIsLoaded)
 	{
-		UE_LOG(LogNeuralNetworkInference, Warning, TEXT("UNeuralNetwork::GetInputDataPointerMutable(): Call UNeuralNetwork::Load() to load a model first."));
-		return nullptr;
+		UE_LOG(LogNeuralNetworkInference, Warning,
+			TEXT("UNeuralNetwork::SetInputFromVoidPointerCopy(): Call UNeuralNetwork::Load() to load a model first."));
 	}
-
+	FNeuralTimer RunTimer;
+	RunTimer.Tic();
 	// UEAndORT
 	if (BackEndForCurrentPlatform == ENeuralBackEnd::UEAndORT)
 	{
-		return ImplBackEndUEAndORT->InputTensors[InTensorIndex].GetData(); // Or ImplBackEndUEAndORT->InputOrtTensors[InTensorIndex].GetTensorMutableData<float>();
+		ImplBackEndUEAndORT->InputTensors[InTensorIndex].SetFromVoidPointerCopy(InVoidPtr);
 	}
-
 	// UEOnly
 	else if (BackEndForCurrentPlatform == ENeuralBackEnd::UEOnly)
 	{
 		FNeuralTensorManager& TensorManager = ImplBackEndUEOnly->TensorManager;
-		return TensorManager.GetTensorsMutable()[TensorManager.GetInputIndexes()[InTensorIndex]].GetData();
+		TensorManager.GetTensorsMutable()[TensorManager.GetInputIndexes()[InTensorIndex]].SetFromVoidPointerCopy(InVoidPtr);
 	}
-
 	// Unknown
-	UE_LOG(LogNeuralNetworkInference, Warning, TEXT("UNeuralNetwork::SetInputFromArrayCopy(): Unknown [BackEnd,BackEndForCurrentPlatform] = [%d,%d]."), (int32)BackEnd, (int32)BackEndForCurrentPlatform);
-	return nullptr;
-}
-
-int64 UNeuralNetwork::GetInputTensorNumber() const
-{
-	// Sanity check
-	if (!bIsLoaded)
-	{
-		UE_LOG(LogNeuralNetworkInference, Warning, TEXT("UNeuralNetwork::GetInputTensorNumber(): Call UNeuralNetwork::Load() to load a model first."));
-		return -1;
-	}
-
-	// UEAndORT
-	if (BackEndForCurrentPlatform == ENeuralBackEnd::UEAndORT)
+	else
 	{
-		return ImplBackEndUEAndORT->InputTensors.Num();
+		UE_LOG(LogNeuralNetworkInference, Warning,
+			TEXT("UNeuralNetwork::SetInputFromVoidPointerCopy(): Unknown [BackEnd,BackEndForCurrentPlatform] = [%d,%d]."),
+			(int32)BackEnd, (int32)BackEndForCurrentPlatform);
 	}
-
-	// UEOnly
-	else if (BackEndForCurrentPlatform == ENeuralBackEnd::UEOnly)
-	{
-		return ImplBackEndUEOnly->TensorManager.GetInputIndexes().Num();
-	}
-
-	// Unknown
-	checkf(false, TEXT("UNeuralNetwork::GetInputTensorNumber(): Unknown [BackEnd,BackEndForCurrentPlatform] = [%d,%d]."), (int32)BackEnd, (int32)BackEndForCurrentPlatform);
-	return -1;
+	InputMemoryTransferStatsModule.StoreSample(RunTimer.Toc());
 }
 
-const FNeuralTensor& UNeuralNetwork::GetOutputTensor(const int32 InTensorIndex) const
+void* UNeuralNetwork::GetInputDataPointerMutable(const int32 InTensorIndex)
 {
-	GET_TENSOR_CODE(InTensorIndex, GetOutputIndexes, GetTensors, OutputTensors);
-}
+	const FScopeLock ResourcesLock(&ResoucesCriticalSection);
 
-int64 UNeuralNetwork::GetOutputTensorNumber() const
-{
 	// Sanity check
 	if (!bIsLoaded)
 	{
-		UE_LOG(LogNeuralNetworkInference, Warning, TEXT("UNeuralNetwork::GetOutputTensorNumber(): Call UNeuralNetwork::Load() to load a model first."));
-		return -1;
+		UE_LOG(LogNeuralNetworkInference, Warning,
+			TEXT("UNeuralNetwork::GetInputDataPointerMutable(): Call UNeuralNetwork::Load() to load a model first."));
+		return nullptr;
 	}
 
 	// UEAndORT
 	if (BackEndForCurrentPlatform == ENeuralBackEnd::UEAndORT)
 	{
-		return ImplBackEndUEAndORT->OutputTensors.Num();
+		return ImplBackEndUEAndORT->InputTensors[InTensorIndex].GetData();
+		// Or return ImplBackEndUEAndORT->InputOrtTensors[InTensorIndex].GetTensorMutableData<float>();
 	}
 
 	// UEOnly
 	else if (BackEndForCurrentPlatform == ENeuralBackEnd::UEOnly)
 	{
-		return ImplBackEndUEOnly->TensorManager.GetOutputIndexes().Num();
+		FNeuralTensorManager& TensorManager = ImplBackEndUEOnly->TensorManager;
+		return TensorManager.GetTensorsMutable()[TensorManager.GetInputIndexes()[InTensorIndex]].GetData();
 	}
 
 	// Unknown
-	UE_LOG(LogNeuralNetworkInference, Warning, TEXT("UNeuralNetwork::GetOutputTensorNumber(): Unknown [BackEnd,BackEndForCurrentPlatform] = [%d,%d]."), (int32)BackEnd, (int32)BackEndForCurrentPlatform);
-	return -1;
-}
-
-TArray<FNeuralTensor> UNeuralNetwork::CreateInputArrayCopy() const
-{
-	TArray<FNeuralTensor> InputTensorArray;
-	for (uint32 InputTensorIndex = 0; InputTensorIndex < GetInputTensorNumber(); ++InputTensorIndex)
-	{
-		const FNeuralTensor& InputTensor = GetInputTensor(InputTensorIndex);
-		InputTensorArray.Push(InputTensor);
-	}
-	return InputTensorArray;
-}
-
-void UNeuralNetwork::SetInputFromArrayCopy(const TArray<FNeuralTensor>& InInputTensorArray)
-{
-	const FScopeLock ResourcesLock(&ResoucesCriticalSection);
-	if (GetInputTensorNumber() != InInputTensorArray.Num())
-	{
-		UE_LOG(LogNeuralNetworkInference, Warning, TEXT("UNeuralNetwork::SetInputFromArrayCopy(): GetInputTensorNumber() == InInputTensorArray.Num() failed, %d != %d."), GetInputTensorNumber(), InInputTensorArray.Num());
-		return;
-	}
-	FNeuralTimer RunTimer;
-	RunTimer.Tic();
-	for (uint32 InputTensorIndex = 0; InputTensorIndex < GetInputTensorNumber(); ++InputTensorIndex)
-	{
-		FNeuralTensor& InputTensor = GetInputTensorMutable(InputTensorIndex);
-		InputTensor.SetFromUnderlyingUInt8ArrayCopy(InInputTensorArray[InputTensorIndex].GetUnderlyingUInt8ArrayRef());
-	}
-	InputMemoryTransferStatsModule.StoreSample(RunTimer.Toc());
-}
-
-TArray<FNeuralTensor> UNeuralNetwork::CreateOutputArrayCopy() const
-{
-	TArray<FNeuralTensor> OutputTensorArray;
-	for (uint32 OutputTensorIndex = 0; OutputTensorIndex < GetOutputTensorNumber(); ++OutputTensorIndex)
-	{
-		const FNeuralTensor& OutputTensor = GetOutputTensor(OutputTensorIndex);
-		OutputTensorArray.Push(OutputTensor);
-	}
-	return OutputTensorArray;
+	UE_LOG(LogNeuralNetworkInference, Warning,
+		TEXT("UNeuralNetwork::SetInputFromArrayCopy(): Unknown [BackEnd,BackEndForCurrentPlatform] = [%d,%d]."),
+		(int32)BackEnd, (int32)BackEndForCurrentPlatform);
+	return nullptr;
 }
 
 void UNeuralNetwork::InputTensorsToGPU(const TArray<int32>& InTensorIndexes)
@@ -393,7 +393,8 @@ void UNeuralNetwork::InputTensorsToGPU(const TArray<int32>& InTensorIndexes)
 	// Sanity check
 	if (!bIsLoaded)
 	{
-		UE_LOG(LogNeuralNetworkInference, Warning, TEXT("UNeuralNetwork::InputTensorsToGPU(): Call UNeuralNetwork::Load() to load a model first."));
+		UE_LOG(LogNeuralNetworkInference, Warning,
+			TEXT("UNeuralNetwork::InputTensorsToGPU(): Call UNeuralNetwork::Load() to load a model first."));
 		return;
 	}
 
@@ -435,7 +436,8 @@ void UNeuralNetwork::OutputTensorsToCPU(const TArray<int32>& InTensorIndexes)
 	// Sanity check
 	if (!bIsLoaded)
 	{
-		UE_LOG(LogNeuralNetworkInference, Warning, TEXT("UNeuralNetwork::OutputTensorToCPU(): Call UNeuralNetwork::Load() to load a model first."));
+		UE_LOG(LogNeuralNetworkInference, Warning,
+			TEXT("UNeuralNetwork::OutputTensorToCPU(): Call UNeuralNetwork::Load() to load a model first."));
 		return;
 	}
 
@@ -498,7 +500,8 @@ void UNeuralNetwork::Run()
 	// Unknown
 	else
 	{
-		UE_LOG(LogNeuralNetworkInference, Warning, TEXT("UNeuralNetwork::Run(): Unknown [BackEnd,BackEndForCurrentPlatform] = [%d,%d]."), (int32)BackEnd, (int32)BackEndForCurrentPlatform);
+		UE_LOG(LogNeuralNetworkInference, Warning, TEXT("UNeuralNetwork::Run(): Unknown [BackEnd,BackEndForCurrentPlatform] = [%d,%d]."),
+			(int32)BackEnd, (int32)BackEndForCurrentPlatform);
 	}
 	ComputeStatsModule.StoreSample(RunTimer.Toc());
 }
@@ -518,6 +521,12 @@ FNeuralStatsData UNeuralNetwork::GetInputMemoryTransferStats() const
 	return InputMemoryTransferStatsModule.GetStats();
 }
 
+void UNeuralNetwork::ResetStats()
+{
+	ComputeStatsModule.ResetStats();
+	InputMemoryTransferStatsModule.ResetStats();
+}
+
 
 /* UNeuralNetwork private functions
  *****************************************************************************/
@@ -535,8 +544,9 @@ bool UNeuralNetwork::Load()
 	if (BackEndForCurrentPlatform == ENeuralBackEnd::UEAndORT)
 	{
 		UNeuralNetwork::FImplBackEndUEAndORT::WarnAndSetDeviceToCPUIfDX12NotEnabled(DeviceType, /*bShouldOpenMessageLog*/true);
-		bIsLoaded = UNeuralNetwork::FImplBackEndUEAndORT::Load(ImplBackEndUEAndORT, OnAsyncRunCompletedDelegate, ThreadModeDelegateForAsyncRunCompleted, ResoucesCriticalSection, AreInputTensorSizesVariable,
-			ModelReadFromFileInBytes, ModelFullFilePath, GetDeviceType(), GetInputDeviceType(), GetOutputDeviceType());
+		bIsLoaded = UNeuralNetwork::FImplBackEndUEAndORT::Load(ImplBackEndUEAndORT, OnAsyncRunCompletedDelegate,
+			ThreadModeDelegateForAsyncRunCompleted, ResoucesCriticalSection, AreInputTensorSizesVariable, ModelReadFromFileInBytes,
+			ModelFullFilePath, GetDeviceType(), GetInputDeviceType(), GetOutputDeviceType());
 	}
 	// UEOnly
 	else if (BackEndForCurrentPlatform == ENeuralBackEnd::UEOnly)
@@ -546,7 +556,8 @@ bool UNeuralNetwork::Load()
 	// Unknown
 	else
 	{
-		UE_LOG(LogNeuralNetworkInference, Warning, TEXT("UNeuralNetwork::Load(): Unknown [BackEnd,BackEndForCurrentPlatform] = [%d,%d]."), (int32)BackEnd, (int32)BackEndForCurrentPlatform);
+		UE_LOG(LogNeuralNetworkInference, Warning, TEXT("UNeuralNetwork::Load(): Unknown [BackEnd,BackEndForCurrentPlatform] = [%d,%d]."),
+			(int32)BackEnd, (int32)BackEndForCurrentPlatform);
 	}
 
 	// Reset Stats
@@ -565,7 +576,8 @@ FNeuralTensor& UNeuralNetwork::GetOutputTensorMutable(const int32 InTensorIndex)
 	GET_TENSOR_CODE(InTensorIndex, GetOutputIndexes, GetTensorsMutable, OutputTensors);
 }
 
-bool UNeuralNetwork::Load(TArray<FNeuralTensor>& InTensors, const TArray<FNeuralTensor*>& InInputTensors, const TArray<FNeuralTensor*>& InOutputTensors, const TArray<TSharedPtr<class FNeuralOperator>>& InOperators)
+bool UNeuralNetwork::Load(TArray<FNeuralTensor>& InTensors, const TArray<FNeuralTensor*>& InInputTensors,
+	const TArray<FNeuralTensor*>& InOutputTensors, const TArray<TSharedPtr<class FNeuralOperator>>& InOperators)
 {
 	DECLARE_SCOPE_CYCLE_COUNTER(TEXT("UNeuralNetwork_Load_FromTensorManagerAndOperators"), STAT_UNeuralNetwork_Load, STATGROUP_MachineLearning);
 
@@ -600,7 +612,8 @@ bool UNeuralNetwork::SetBackEnd(const UNeuralNetwork::ENeuralBackEnd InBackEnd)
 	BackEnd = InBackEnd;
 	const ENeuralBackEnd NewBackEndForCurrentPlatform = FPrivateNeuralNetwork::SetBackEndForCurrentPlatform(BackEnd);
 	// Reload only required if BackEndForCurrentPlatform changes (regardless of whether BackEnd changed).
-	// BackEndForCurrentPlatform does not necessarily change if BackEnd changes. E.g., changing from UEAndORT into Auto in Windows will result in BackEndForCurrentPlatform = UEAndORT in both cases.
+	// BackEndForCurrentPlatform does not necessarily change if BackEnd changes. E.g., changing from UEAndORT into Auto in Windows will result in
+	// BackEndForCurrentPlatform = UEAndORT in both cases.
 	if (BackEndForCurrentPlatform != NewBackEndForCurrentPlatform)
 	{
 		BackEndForCurrentPlatform = NewBackEndForCurrentPlatform;
diff --git a/Engine/Plugins/Experimental/NeuralNetworkInference/Source/NeuralNetworkInference/Private/NeuralTensor.cpp b/Engine/Plugins/Experimental/NeuralNetworkInference/Source/NeuralNetworkInference/Private/NeuralTensor.cpp
--- a/Engine/Plugins/Experimental/NeuralNetworkInference/Source/NeuralNetworkInference/Private/NeuralTensor.cpp
+++ b/Engine/Plugins/Experimental/NeuralNetworkInference/Source/NeuralNetworkInference/Private/NeuralTensor.cpp
@@ -164,12 +164,13 @@ FNeuralTensor::FNeuralTensor(const ENeuralDataType InDataType, const void* const
 	// Sanity check
 	if (IsEmpty())
 	{
-		UE_LOG(LogNeuralNetworkInference, Warning, TEXT("FNeuralTensor(): GetVolume() = 0. Skipping array copy."));
+		UE_LOG(LogNeuralNetworkInference, Warning,
+			TEXT("FNeuralTensor(): GetVolume() = 0, skipping array copy. Use a different constructor if the volume is 0 to avoid this warning."));
 	}
 	// Memory filling
 	else
 	{
-		SetFromPointer(InValues, InSizeOfT, InValueNum);
+		SetFromVoidPointerCopy(InValues, InSizeOfT, InValueNum);
 	}
 }
 
@@ -306,6 +307,14 @@ void FNeuralTensor::SetFromUnderlyingUInt8ArrayCopy(const TArray<uint8>& InArray
 	UnderlyingUInt8ArrayData = InArray;
 }
 
+void FNeuralTensor::SetFromVoidPointerCopy(const void* const InVoidPtr)
+{
+	if (InVoidPtr)
+	{
+		FMemory::Memcpy(UnderlyingUInt8ArrayData.GetData(), InVoidPtr, NumInBytes());
+	}
+}
+
 bool FNeuralTensor::SetFromTensorProto(const FTensorProto* const InTensorProto, const FString& InTensorName, const ENeuralTensorType InTensorType)
 {
 	if (!InTensorProto)
@@ -812,19 +821,19 @@ FString FNeuralTensor::ToString(const int64 InMaxNumberElementsToDisplay, const
 /* FNeuralTensor private functions
  *****************************************************************************/
 
-void FNeuralTensor::SetFromPointer(const void* const InData, const int64 InSizeOfT, const int64 InDataSize)
+void FNeuralTensor::SetFromVoidPointerCopy(const void* const InVoidPtr, const int64 InSizeOfT, const int64 InDataSize)
 {
 	// Sanity checks
 	if (Num() != InDataSize || NumInBytes() != InSizeOfT * InDataSize)
 	{
-		UE_LOG(LogNeuralNetworkInference, Warning, TEXT("FNeuralTensor-%s::SetFromPointer: Num() == InDataSize failed, %d vs. %d, or"
+		UE_LOG(LogNeuralNetworkInference, Warning, TEXT("FNeuralTensor-%s::SetFromVoidPointerCopy: Num() == InDataSize failed, %d vs. %d, or"
 			" NumInBytes() == sizeof(T) x InDataSize failed, %d vs. %d. If you want to modify the dimensions of FNeuralTensor, call"
 			" SetNumUninitialized() first."), *Name, Num(), InDataSize, NumInBytes(), InSizeOfT * InDataSize);
 	}
 	else
 	{
 		// Deep copy
-		FMemory::Memcpy(UnderlyingUInt8ArrayData.GetData(), InData, NumInBytes());
+		SetFromVoidPointerCopy(InVoidPtr);
 	}
 }
 
diff --git a/Engine/Plugins/Experimental/NeuralNetworkInference/Source/NeuralNetworkInference/Public/NeuralNetwork.h b/Engine/Plugins/Experimental/NeuralNetworkInference/Source/NeuralNetworkInference/Public/NeuralNetwork.h
--- a/Engine/Plugins/Experimental/NeuralNetworkInference/Source/NeuralNetworkInference/Public/NeuralNetwork.h
+++ b/Engine/Plugins/Experimental/NeuralNetworkInference/Source/NeuralNetworkInference/Public/NeuralNetwork.h
@@ -108,10 +108,10 @@ public:
 
 	/**
 	 * Getter and setter functions for DeviceType, InputDeviceType, and OutputDeviceType:
-	 * - GetDeviceType(): It returns DeviceType.
-	 * - GetInputDeviceType(): It returns InputDeviceType.
-	 * - GetOutputDeviceType(): It returns OutputDeviceType.
-	 * - SetDeviceType(): It sets DeviceType, InputDeviceType, and OutputDeviceType. If you are setting it to GPU, check IsGPUSupported() first.
+	 * - GetDeviceType() returns DeviceType.
+	 * - GetInputDeviceType() returns InputDeviceType.
+	 * - GetOutputDeviceType() returns OutputDeviceType.
+	 * - SetDeviceType() sets DeviceType, InputDeviceType, and OutputDeviceType. If you are setting it to GPU, check IsGPUSupported() first.
 	 * @see DeviceType, InputDeviceType, OutputDeviceType, and IsGPUSupported() for more details.
 	 */
 	ENeuralDeviceType GetDeviceType() const;
@@ -136,8 +136,8 @@ public:
 
 	/**
 	 * Getter and setter functions for SynchronousMode:
-	 * - GetDeviceType(): It returns SynchronousMode.
-	 * - SetDeviceType(): It sets SynchronousMode.
+	 * - GetDeviceType() returns SynchronousMode.
+	 * - SetDeviceType() sets SynchronousMode.
 	 * @see SynchronousMode and GetOnAsyncRunCompletedDelegate() for more details.
 	 */
 	ENeuralSynchronousMode GetSynchronousMode() const;
@@ -164,37 +164,43 @@ public:
 
 	/**
 	 * Getter and setter functions for ThreadModeDelegateForAsyncRunCompleted:
-	 * - GetThreadModeDelegateForAsyncRunCompleted(): It returns ThreadModeDelegateForAsyncRunCompleted.
-	 * - SetThreadModeDelegateForAsyncRunCompleted(): It sets ThreadModeDelegateForAsyncRunCompleted.
+	 * - GetThreadModeDelegateForAsyncRunCompleted() returns ThreadModeDelegateForAsyncRunCompleted.
+	 * - SetThreadModeDelegateForAsyncRunCompleted() sets ThreadModeDelegateForAsyncRunCompleted.
 	 * @see ThreadModeDelegateForAsyncRunCompleted and GetOnAsyncRunCompletedDelegate() for more details.
 	 */
 	ENeuralThreadMode GetThreadModeDelegateForAsyncRunCompleted() const;
 	void SetThreadModeDelegateForAsyncRunCompleted(const ENeuralThreadMode InThreadModeDelegateForAsyncRunCompleted);
 
-// CONTINUE HERE
 	/**
-	 * Functions to get/fill input.
-	 * These functions either take a TArray as input, return a modifiable void* to fill the data, or return a constant FNeuralTensor(s) to see input
-	 * properties (e.g., size or dimensions).
+	 * GetInputTensorNumber() and GetOutputTensorNumber() return the number of input or output tensors of this network, respectively.
+	 * @see GetInputTensor(), SetInputFromArrayCopy(), and GetInputDataPointerMutable() for other input-tensor-related functions.
+	 * @see GetOutputTensor() for other output-tensor-related functions.
 	 */
-	const FNeuralTensor& GetInputTensor(const int32 InTensorIndex = 0) const;
-	void SetInputFromArrayCopy(const TArray<float>& InArray, const int32 InTensorIndex = 0);
-	void* GetInputDataPointerMutable(const int32 InTensorIndex = 0);
 	int64 GetInputTensorNumber() const;
+	int64 GetOutputTensorNumber() const;
 
 	/**
-	 * Slow functions (as they will copy every input/output FNeuralNetwork) only meant for debugging purposes.
+	 * GetInputTensor() and GetOutputTensor() return a const (read-only) reference of the input or output FNeuralTensor(s) of the network,
+	 * respectively. They allow querying tensor properties, such as size, GetInputTensor().GetSize(), or dimensions, GetInputTensor().Num().
+	 * @param InTensorIndex The input/output index to query. By default (0), it queries the first input/output tensor. @see GetInputTensorNumber()
+	 * and GetOutputTensorNumber() for more details.
 	 */
-	TArray<FNeuralTensor> CreateInputArrayCopy() const;
-	void SetInputFromArrayCopy(const TArray<FNeuralTensor>& InTensorDataArray);
-	TArray<FNeuralTensor> CreateOutputArrayCopy() const;
+	const FNeuralTensor& GetInputTensor(const int32 InTensorIndex = 0) const;
+	const FNeuralTensor& GetOutputTensor(const int32 InTensorIndex = 0) const;
 
 	/**
-	 * Functions to get output. The returned FNeuralTensor(s) are constant to prevent the user from modifying the tensor properties (e.g., size or
-	 * dimensions).
+	 * SetInputFromArrayCopy() and GetInputDataPointerMutable() are the only functions that allow modifing the network input tensor(s) values:
+	 * - SetInputFromArrayCopy() is very easy to use but less efficient (it requires an intermidiate and auxiliary TArray). This function copies the
+	 *   given InArray into the desired input FNeuralTensor.
+	 * - GetInputDataPointerMutable() is potentially more efficient because it avoids creating and copying from an intermediate TArray. This function
+	 *   returns a mutable void* pointer of the desired input FNeuralTensor that can be filled by the user on-the-fly.
 	 */
-	const FNeuralTensor& GetOutputTensor(const int32 InTensorIndex = 0) const;
-	int64 GetOutputTensorNumber() const;
+	void SetInputFromArrayCopy(const TArray<float>& InArray, const int32 InTensorIndex = 0);
+// CONTINUE HERE
+	void SetInputFromVoidPointerCopy(const void* const InVoidPtr, const int32 InTensorIndex = 0);
+	void* GetInputDataPointerMutable(const int32 InTensorIndex = 0);
+
+// @todo: Move these 3 functions to NNIQA.
 
 	/**
 	 * Non-efficient functions meant to be used only for debugging purposes.
@@ -320,7 +326,11 @@ private:
 	bool Load();
 
 	/**
-	 * Private and mutable version of GetInputTensor()/GetOutputTensor().
+	 * GetInputTensorMutable() and GetOutputTensorMutable() are the private and mutable version of GetInputTensor() and GetOutputTensor(), respectively.
+	 * Most FNeuralTensor properties are considered constant and fixed after calling UNeuralNetwork::Load(), other than the actual data values
+	 * (i.e., UnderlyingUInt8ArrayData) and some GPU properties. Thus, use these 2 functions with extreme precaution and be sure not to modify any
+	 * other FNeuralNetwork properties.
+	 * @see GetInputTensor() and GetOutputTensor() for more details.
 	 */
 	FNeuralTensor& GetInputTensorMutable(const int32 InTensorIndex = 0);
 	FNeuralTensor& GetOutputTensorMutable(const int32 InTensorIndex = 0);
diff --git a/Engine/Plugins/Experimental/NeuralNetworkInference/Source/NeuralNetworkInference/Public/NeuralTensor.h b/Engine/Plugins/Experimental/NeuralNetworkInference/Source/NeuralNetworkInference/Public/NeuralTensor.h
--- a/Engine/Plugins/Experimental/NeuralNetworkInference/Source/NeuralNetworkInference/Public/NeuralTensor.h
+++ b/Engine/Plugins/Experimental/NeuralNetworkInference/Source/NeuralNetworkInference/Public/NeuralTensor.h
@@ -43,9 +43,9 @@ public:
 	 * @param InTensorType It sets TensorType. @see TensorType for more details.
 	 * @param InArray Alternative to InSizes/InVolume, and represents the input data to copy from. The template constructor that uses InArray is a
 	 *  simple (but not efficient) way of initializing a FNeuralTensor from the existing InArray TArray. It is equivalent to calling
-	 *  FNeuralTensor(..., InSizes/InVolume, ...) and then SetFromArrayCopy(InArray). It is not efficient because it makes a deep copy of the data
-	 *  of InArray. For maximum speed, rather than creating this intermediate InArray TArray, you can avoid the double copy by calling
-	 *  FNeuralTensor(..., InSizes/InVolume, ...) and then filling the tensor memory with GetData()/GetDataCasted().
+	 *  FNeuralTensor(..., InSizes/InVolume, ...) and then SetFromArrayCopy(InArray)/SetFromVoidPointerCopy(InVoidPtr). It is not efficient because
+	 *  it makes a deep copy of the data of InArray. For maximum speed, rather than creating this intermediate InArray TArray, you can avoid the
+	 *  double copy by calling FNeuralTensor(..., InSizes/InVolume, ...) and then filling the tensor memory with GetData()/GetDataCasted().
 	 *
 	 * In addition, UnderlyingUInt8ArrayData is preallocated from InDataType and InSizes/InVolume. @see UnderlyingUInt8ArrayData for more details.
 	 */
@@ -139,9 +139,9 @@ public:
 	FORCEINLINE void SetEnableGPU(const bool bInEnableGPU);
 
 	/**
-	 * @see At(), GetArrayCopy(), GetUnderlyingUInt8ArrayRef(), GetData(), GetDataCasted(), SetNumUninitialized, SetFromUnderlyingUInt8ArrayCopy,
-	 * SetFromArrayCopy, SetTo, SetFromTensorProto, which represent all the functions that can access/modify the CPU memory. Some of these functions
-	 * could result in undefined behavior in not used properly:
+	 * @see At(), GetArrayCopy(), GetUnderlyingUInt8ArrayRef(), GetData(), GetDataCasted(), SetNumUninitialized(), SetFromUnderlyingUInt8ArrayCopy(),
+	 * SetFromArrayCopy(), SetFromVoidPointerCopy(), SetTo(), SetFromTensorProto(), which represent all the functions that can access/modify the CPU
+	 * memory. Some of these functions could result in undefined behavior in not used properly:
 	 * - If FNeuralTensor goes out of scope, it is destructed, or its memory is reset (constructor, SetNum(), etc.), the functions returning
 	 *   references/pointers will no longer be valid.
 	 * - These functions only modify the CPU memory, they do not synchronize CPU and GPU memory automatically. The user must do this accordingly.
@@ -157,9 +157,9 @@ public:
 	const T& At(const TInput InIndex) const;
 
 	/**
-	 * @see At(), GetArrayCopy(), GetUnderlyingUInt8ArrayRef(), GetData(), GetDataCasted(), SetNumUninitialized, SetFromUnderlyingUInt8ArrayCopy,
-	 * SetFromArrayCopy, SetTo, SetFromTensorProto, which represent all the functions that can access/modify the CPU memory. Some of these functions
-	 * could result in undefined behavior in not used properly:
+	 * @see At(), GetArrayCopy(), GetUnderlyingUInt8ArrayRef(), GetData(), GetDataCasted(), SetNumUninitialized(), SetFromUnderlyingUInt8ArrayCopy(),
+	 * SetFromArrayCopy(), SetFromVoidPointerCopy(), SetTo(), SetFromTensorProto(), which represent all the functions that can access/modify the CPU
+	 * memory. Some of these functions could result in undefined behavior in not used properly:
 	 * - If FNeuralTensor goes out of scope, it is destructed, or its memory is reset (constructor, SetNum(), etc.), the functions returning
 	 *   references/pointers will no longer be valid.
 	 * - These functions only modify the CPU memory, they do not synchronize CPU and GPU memory automatically. The user must do this accordingly.
@@ -173,9 +173,9 @@ public:
 	FORCEINLINE const TArray<uint8>& GetUnderlyingUInt8ArrayRef() const;
 
 	/**
-	 * @see At(), GetArrayCopy(), GetUnderlyingUInt8ArrayRef(), GetData(), GetDataCasted(), SetNumUninitialized, SetFromUnderlyingUInt8ArrayCopy,
-	 * SetFromArrayCopy, SetTo, SetFromTensorProto, which represent all the functions that can access/modify the CPU memory. Some of these functions
-	 * could result in undefined behavior in not used properly:
+	 * @see At(), GetArrayCopy(), GetUnderlyingUInt8ArrayRef(), GetData(), GetDataCasted(), SetNumUninitialized(), SetFromUnderlyingUInt8ArrayCopy(),
+	 * SetFromArrayCopy(), SetFromVoidPointerCopy(), SetTo(), SetFromTensorProto(), which represent all the functions that can access/modify the CPU
+	 * memory. Some of these functions could result in undefined behavior in not used properly:
 	 * - If FNeuralTensor goes out of scope, it is destructed, or its memory is reset (constructor, SetNum(), etc.), the functions returning
 	 *   references/pointers will no longer be valid.
 	 * - These functions only modify the CPU memory, they do not synchronize CPU and GPU memory automatically. The user must do this accordingly.
@@ -197,9 +197,9 @@ public:
 	FORCEINLINE const T* const GetDataCasted() const;
 
 	/**
-	 * @see At(), GetArrayCopy(), GetUnderlyingUInt8ArrayRef(), GetData(), GetDataCasted(), SetNumUninitialized, SetFromUnderlyingUInt8ArrayCopy,
-	 * SetFromArrayCopy, SetTo, SetFromTensorProto, which represent all the functions that can access/modify the CPU memory. Some of these functions
-	 * could result in undefined behavior in not used properly:
+	 * @see At(), GetArrayCopy(), GetUnderlyingUInt8ArrayRef(), GetData(), GetDataCasted(), SetNumUninitialized(), SetFromUnderlyingUInt8ArrayCopy(),
+	 * SetFromArrayCopy(), SetFromVoidPointerCopy(), SetTo(), SetFromTensorProto(), which represent all the functions that can access/modify the CPU
+	 * memory. Some of these functions could result in undefined behavior in not used properly:
 	 * - If FNeuralTensor goes out of scope, it is destructed, or its memory is reset (constructor, SetNum(), etc.), the functions returning
 	 *   references/pointers will no longer be valid.
 	 * - These functions only modify the CPU memory, they do not synchronize CPU and GPU memory automatically. The user must do this accordingly.
@@ -218,16 +218,17 @@ public:
 	FORCEINLINE void SetNumUninitialized(const ENeuralDataType InDataType, const int64 InVolume, const bool bInAllowShrinking = true);
 
 	/**
-	 * @see At(), GetArrayCopy(), GetUnderlyingUInt8ArrayRef(), GetData(), GetDataCasted(), SetNumUninitialized, SetFromUnderlyingUInt8ArrayCopy,
-	 * SetFromArrayCopy, SetTo, SetFromTensorProto, which represent all the functions that can access/modify the CPU memory. Some of these functions
-	 * could result in undefined behavior in not used properly:
+	 * @see At(), GetArrayCopy(), GetUnderlyingUInt8ArrayRef(), GetData(), GetDataCasted(), SetNumUninitialized(), SetFromUnderlyingUInt8ArrayCopy(),
+	 * SetFromArrayCopy(), SetFromVoidPointerCopy(), SetTo(), SetFromTensorProto(), which represent all the functions that can access/modify the CPU
+	 * memory. Some of these functions could result in undefined behavior in not used properly:
 	 * - If FNeuralTensor goes out of scope, it is destructed, or its memory is reset (constructor, SetNum(), etc.), the functions returning
 	 *   references/pointers will no longer be valid.
 	 * - These functions only modify the CPU memory, they do not synchronize CPU and GPU memory automatically. The user must do this accordingly.
 	 *
-	 * SetFromUnderlyingUInt8ArrayCopy/SetFromArrayCopy will replace the contents of the underlying CPU TArray with the input one, by deeply copying
-	 * the array (safer and easier to use). For maximum performance:
-	 * - If you are given an input TArray (or a FNeuralTensor with different Sizes) that cannot be moved, use SetFromArrayCopy().
+	 * SetFromUnderlyingUInt8ArrayCopy/SetFromArrayCopy/SetFromVoidPointerCopy will replace the contents of the underlying CPU TArray with the input
+	 * TArray or void pointer, by deeply copying the array/pointer (safer and easier to use). For maximum performance:
+	 * - If you are given an input TArray/pointer (or a FNeuralTensor with different Sizes) that cannot be moved, use SetFromArrayCopy() or
+	 *   SetFromVoidPointerCopy().
 	 * - If you are initializing the FNeuralTensor iteratively (and to avoid a copy), fill the tensor with GetData()/GetDataCasted<T>().
 	 * Modifying InArray after calling this function will not modify the data of this FNeuralTensor.
 	 * The size of InArray and this tensor must match, i.e., NumInBytes() must match TArray<uint8>::Num() or Num() match TArray<T>.Num().
@@ -235,11 +236,12 @@ public:
 	void SetFromUnderlyingUInt8ArrayCopy(const TArray<uint8>& InArray);
 	template<typename T>
 	void SetFromArrayCopy(const TArray<T>& InArray);
+	void SetFromVoidPointerCopy(const void* const InVoidPtr);
 
 	/**
-	 * @see At(), GetArrayCopy(), GetUnderlyingUInt8ArrayRef(), GetData(), GetDataCasted(), SetNumUninitialized, SetFromUnderlyingUInt8ArrayCopy,
-	 * SetFromArrayCopy, SetTo, SetFromTensorProto, which represent all the functions that can access/modify the CPU memory. Some of these functions
-	 * could result in undefined behavior in not used properly:
+	 * @see At(), GetArrayCopy(), GetUnderlyingUInt8ArrayRef(), GetData(), GetDataCasted(), SetNumUninitialized(), SetFromUnderlyingUInt8ArrayCopy(),
+	 * SetFromArrayCopy(), SetFromVoidPointerCopy(), SetTo(), SetFromTensorProto(), which represent all the functions that can access/modify the CPU
+	 * memory. Some of these functions could result in undefined behavior in not used properly:
 	 * - If FNeuralTensor goes out of scope, it is destructed, or its memory is reset (constructor, SetNum(), etc.), the functions returning
 	 *   references/pointers will no longer be valid.
 	 * - These functions only modify the CPU memory, they do not synchronize CPU and GPU memory automatically. The user must do this accordingly.
@@ -251,9 +253,9 @@ public:
 	void SetTo(const TInput InValue);
 
 	/**
-	 * @see At(), GetArrayCopy(), GetUnderlyingUInt8ArrayRef(), GetData(), GetDataCasted(), SetNumUninitialized, SetFromUnderlyingUInt8ArrayCopy,
-	 * SetFromArrayCopy, SetTo, SetFromTensorProto, which represent all the functions that can access/modify the CPU memory. Some of these functions
-	 * could result in undefined behavior in not used properly:
+	 * @see At(), GetArrayCopy(), GetUnderlyingUInt8ArrayRef(), GetData(), GetDataCasted(), SetNumUninitialized(), SetFromUnderlyingUInt8ArrayCopy(),
+	 * SetFromArrayCopy(), SetFromVoidPointerCopy(), SetTo(), SetFromTensorProto(), which represent all the functions that can access/modify the CPU
+	 * memory. Some of these functions could result in undefined behavior in not used properly:
 	 * - If FNeuralTensor goes out of scope, it is destructed, or its memory is reset (constructor, SetNum(), etc.), the functions returning
 	 *   references/pointers will no longer be valid.
 	 * - These functions only modify the CPU memory, they do not synchronize CPU and GPU memory automatically. The user must do this accordingly.
@@ -393,7 +395,7 @@ private:
 	 * Auxiliary function for the template SetFromArrayCopy().
 	 * It simply moves as much code as possible from the template function to the cpp file.
 	 */
-	void SetFromPointer(const void* const InData, const int64 InSizeOfT, const int64 InDataSize);
+	void SetFromVoidPointerCopy(const void* const InVoidPtr, const int64 InSizeOfT, const int64 InDataSize);
 
 	/**
 	 * Checks and warns whether the current data type T is the same than the one defined by DataType.
@@ -539,7 +541,7 @@ void FNeuralTensor::SetFromArrayCopy(const TArray<T>& InArray)
 {
 	if (CheckTAndDataTypeEquivalent<T>())
 	{
-		SetFromPointer(InArray.GetData(), sizeof(T), InArray.Num());
+		SetFromVoidPointerCopy(InArray.GetData(), sizeof(T), InArray.Num());
 	}
 }
 
diff --git a/Engine/Plugins/Experimental/NeuralNetworkInference/Source/Private/NeuralNetworkInferenceQA/Private/NeuralNetworkInferenceQAAsset.cpp b/Engine/Plugins/Experimental/NeuralNetworkInference/Source/Private/NeuralNetworkInferenceQA/Private/NeuralNetworkInferenceQAAsset.cpp
--- a/Engine/Plugins/Experimental/NeuralNetworkInference/Source/Private/NeuralNetworkInferenceQA/Private/NeuralNetworkInferenceQAAsset.cpp
+++ b/Engine/Plugins/Experimental/NeuralNetworkInference/Source/Private/NeuralNetworkInferenceQA/Private/NeuralNetworkInferenceQAAsset.cpp
@@ -27,15 +27,17 @@ void FNeuralNetworkInferenceQAOperatorTestAsset::AddOutputTensors(const UNeuralN
 	}
 }
 
-bool FNeuralNetworkInferenceQAOperatorTestAsset::CompareAverageL1DiffNewVsPreviousTests(const FNeuralNetworkInferenceQAOperatorTestAsset& InOperatorTestAsset1,
-	const FNeuralNetworkInferenceQAOperatorTestAsset& InOperatorTestAsset2, const FString& InOperatorName)
+bool FNeuralNetworkInferenceQAOperatorTestAsset::CompareAverageL1DiffNewVsPreviousTests(
+	const FNeuralNetworkInferenceQAOperatorTestAsset& InOperatorTestAsset1, const FNeuralNetworkInferenceQAOperatorTestAsset& InOperatorTestAsset2,
+	const FString& InOperatorName)
 {
 	// Compare each input and output
 	return (CompareAverageL1DiffNewVsPreviousTests(InOperatorTestAsset1.InputTensors, InOperatorTestAsset2.InputTensors, InOperatorName)
 		&& CompareAverageL1DiffNewVsPreviousTests(InOperatorTestAsset1.OutputTensors, InOperatorTestAsset2.OutputTensors, InOperatorName));
 }
 
-bool FNeuralNetworkInferenceQAOperatorTestAsset::CompareAverageL1DiffNewVsPreviousTests(const TArray<FNeuralTensor>& InTensorsA, const TArray<FNeuralTensor>& InTensorsB, const FString& InOperatorName)
+bool FNeuralNetworkInferenceQAOperatorTestAsset::CompareAverageL1DiffNewVsPreviousTests(const TArray<FNeuralTensor>& InTensorsA,
+	const TArray<FNeuralTensor>& InTensorsB, const FString& InOperatorName)
 {
 	// Number of inputs and outputs should be the same
 	if (InTensorsA.Num() != InTensorsB.Num())
@@ -46,7 +48,8 @@ bool FNeuralNetworkInferenceQAOperatorTestAsset::CompareAverageL1DiffNewVsPrevio
 	const FString DebugName = InOperatorName + TEXT("-CPUvsGT");
 	for (int32 TensorIndex = 0; TensorIndex < InTensorsA.Num(); ++TensorIndex)
 	{
-		if (!FNeuralNetworkInferenceQAUtils::EstimateTensorL1DiffError(InTensorsA[TensorIndex], InTensorsB[TensorIndex], /*ZeroThreshold*/5e-4, DebugName))
+		if (!FNeuralNetworkInferenceQAUtils::EstimateTensorL1DiffError(InTensorsA[TensorIndex], InTensorsB[TensorIndex], /*ZeroThreshold*/5e-4,
+		DebugName))
 		{
 			return false;
 		}
@@ -62,7 +65,7 @@ bool FNeuralNetworkInferenceQAOperatorTestAsset::CompareAverageL1DiffNewVsPrevio
 FString FNeuralNetworkInferenceQAOperatorAsset::RunNetworkCPUAndGetString(UNeuralNetwork* InOutNetwork)
 {
 	// Create input string and InputTensorArray
-	const TArray<FNeuralTensor> InputTensorArray = InOutNetwork->CreateInputArrayCopy();
+	const TArray<FNeuralTensor> InputTensorArray = FNeuralNetworkInferenceQAUtils::CreateInputArrayCopy(InOutNetwork);
 	FString TensorsAsString = TEXT("Input(s):\n");
 	for (int32 InputIndex = 0; InputIndex < InOutNetwork->GetInputTensorNumber(); ++InputIndex)
 	{
@@ -78,7 +81,7 @@ FString FNeuralNetworkInferenceQAOperatorAsset::RunNetworkCPUAndGetString(UNeura
 		TensorsAsString += InOutNetwork->GetOutputTensor(OutputIndex).ToString() + TEXT("\n");
 	}
 	// Reset memory
-	InOutNetwork->SetInputFromArrayCopy(InputTensorArray); // This is doing an unnecessary copy (it should move really)
+	FNeuralNetworkInferenceQAUtils::SetInputFromArrayCopy(InOutNetwork, InputTensorArray); // This is doing an unnecessary copy (it should move)
 	// Return TensorsAsString
 	return TensorsAsString + TEXT("\n\n\n");
 }
@@ -106,7 +109,8 @@ bool FNeuralNetworkInferenceQAOperatorAsset::CompareNewVsPreviousTests(const FSt
 		}
 		else
 		{
-			UE_LOG(LogNeuralNetworkInferenceQA, Warning, TEXT("New tests added or removed, PreviousTests.Num() = %d and NewTests.Num() = %d."), PreviousTests.Num(), NewTests.Num());
+			UE_LOG(LogNeuralNetworkInferenceQA, Warning, TEXT("New tests added or removed, PreviousTests.Num() = %d and NewTests.Num() = %d."),
+				PreviousTests.Num(), NewTests.Num());
 		}
 		bWasComparisonSuccessful = false;
 	}
@@ -121,7 +125,8 @@ bool FNeuralNetworkInferenceQAOperatorAsset::CompareNewVsPreviousTests(const FSt
 				// Even if PreviousTestsString != NewTestsString, it might be a floating value precision error
 				const FNeuralNetworkInferenceQAOperatorTestAsset& Test = PreviousTests[TestIndex];
 				const FNeuralNetworkInferenceQAOperatorTestAsset& UpdatedTest = NewTests[TestIndex];
-				bWasComparisonSuccessful &= FNeuralNetworkInferenceQAOperatorTestAsset::CompareAverageL1DiffNewVsPreviousTests(Test, UpdatedTest, InOperatorName);
+				bWasComparisonSuccessful &= FNeuralNetworkInferenceQAOperatorTestAsset::CompareAverageL1DiffNewVsPreviousTests(Test, UpdatedTest,
+					InOperatorName);
 			}
 		}
 	}
@@ -134,7 +139,8 @@ bool FNeuralNetworkInferenceQAOperatorAsset::CompareNewVsPreviousTests(const FSt
 		const FString FilePathNewTest = BaseTestPath + TEXT("_new") + GroundTruthFileExtension;
 		ensureMsgf(FFileHelper::SaveStringToFile(PreviousTestsString, *FilePathPreviousTest), TEXT("FFileHelper::SaveStringToFile returned false."));
 		ensureMsgf(FFileHelper::SaveStringToFile(NewTestsString, *FilePathNewTest), TEXT("FFileHelper::SaveStringToFile returned false."));
-		UE_LOG(LogNeuralNetworkInferenceQA, Warning, TEXT("FNeuralNetworkInferenceQAOperatorAsset::CompareNewVsPreviousTests(): Mismatch between expected and actual results, they should match. Check the following files for differences:\n"
+		UE_LOG(LogNeuralNetworkInferenceQA, Warning, TEXT("FNeuralNetworkInferenceQAOperatorAsset::CompareNewVsPreviousTests(): Mismatch between"
+			" expected and actual results, they should match. Check the following files for differences:\n"
 			"\t- Character length (saved previous vs. new): %d vs. %d\n"
 			"\t- Previous results saved in %s\n"
 			"\t- New results saved in %s\n"),
@@ -171,13 +177,15 @@ void UNeuralNetworkInferenceQAAsset::FindOrAddOperators(const TArray<FString>& I
 		{
 			NotFoundOperatorNames += OperatorName + TEXT(", ");
 			Operators.Add(OperatorName, FNeuralNetworkInferenceQAOperatorAsset());
-			UE_LOG(LogNeuralNetworkInferenceQA, Warning, TEXT("QA for operator %s was not found on UNeuralNetworkInferenceQAAsset, added!"), *OperatorName);
+			UE_LOG(LogNeuralNetworkInferenceQA, Warning, TEXT("QA for operator %s was not found on UNeuralNetworkInferenceQAAsset, added!"),
+				*OperatorName);
 		}
 	}
 	if (NotFoundOperatorNames.Len() > 0)
 	{
 		NotFoundOperatorNames.LeftChopInline(2); // Removes the last TEXT(", ")
-		ensureMsgf(false, TEXT("Some operators are new and were not found on UNeuralNetworkInferenceQAAsset, they have been added: %s."), *NotFoundOperatorNames);
+		ensureMsgf(false, TEXT("Some operators are new and were not found on UNeuralNetworkInferenceQAAsset, they have been added: %s."),
+			*NotFoundOperatorNames);
 	}
 }
 
@@ -199,7 +207,8 @@ bool UNeuralNetworkInferenceQAAsset::CompareNewVsPreviousTests(const FString& In
 	}
 	if (!bWasComparisonSuccessful)
 	{
-		UE_LOG(LogNeuralNetworkInferenceQA, Warning, TEXT("UNeuralNetworkInferenceQAAsset::CompareNewVsPreviousTests(): Mismatch between expected and actual results, they should match. Check the previous warning messages."));
+		UE_LOG(LogNeuralNetworkInferenceQA, Warning, TEXT("UNeuralNetworkInferenceQAAsset::CompareNewVsPreviousTests(): Mismatch between expected"
+			" and actual results, they should match. Check the previous warning messages."));
 	}
 	return bWasComparisonSuccessful;
 }
@@ -214,15 +223,21 @@ void UNeuralNetworkInferenceQAAsset::FlushNewTests()
 
 const FString UNeuralNetworkInferenceQAAsset_FileName = TEXT("NeuralNetworkInferenceQAAsset");
 
-UNeuralNetworkInferenceQAAsset* UNeuralNetworkInferenceQAAsset::Load(const FString& InNeuralNetworkInferenceQAAssetParentDirectoryName, const FString& InNeuralNetworkInferenceQAAssetName)
+UNeuralNetworkInferenceQAAsset* UNeuralNetworkInferenceQAAsset::Load(const FString& InNeuralNetworkInferenceQAAssetParentDirectoryName,
+	const FString& InNeuralNetworkInferenceQAAssetName)
 {
 	// Load NeuralNetworkInferenceQAAsset from disk
-	const FString FilePath = UNeuralNetworkInferenceQAAsset_FileName + TEXT("'/Game/") / InNeuralNetworkInferenceQAAssetParentDirectoryName / UNeuralNetworkInferenceQAAsset_FileName + TEXT(".") + UNeuralNetworkInferenceQAAsset_FileName + TEXT("'"); // "NeuralNetworkInferenceQAAsset'/Game/[UNIT_TEST_DIR]/NeuralNetworkInferenceQAAsset.NeuralNetworkInferenceQAAsset'"
-	UNeuralNetworkInferenceQAAsset* NeuralNetworkInferenceQAAsset = LoadObject<UNeuralNetworkInferenceQAAsset>((UObject*)GetTransientPackage(), *FilePath);
+	// E.g., "NeuralNetworkInferenceQAAsset'/Game/[UNIT_TEST_DIR]/NeuralNetworkInferenceQAAsset.NeuralNetworkInferenceQAAsset'"
+	const FString FilePath = UNeuralNetworkInferenceQAAsset_FileName
+		+ TEXT("'/Game/") / InNeuralNetworkInferenceQAAssetParentDirectoryName / UNeuralNetworkInferenceQAAsset_FileName + TEXT(".")
+		+ UNeuralNetworkInferenceQAAsset_FileName + TEXT("'");
+	UNeuralNetworkInferenceQAAsset* NeuralNetworkInferenceQAAsset = LoadObject<UNeuralNetworkInferenceQAAsset>((UObject*)GetTransientPackage(),
+		*FilePath);
 	// Warning if it does not exist yet
 	if (!NeuralNetworkInferenceQAAsset)
 	{
-		UE_LOG(LogNeuralNetworkInferenceQA, Display, TEXT("NeuralNetworkInferenceQAAsset not found in %s. Please, create it first or make sure the right path is being used."), *FilePath);
+		UE_LOG(LogNeuralNetworkInferenceQA, Display,
+			TEXT("NeuralNetworkInferenceQAAsset not found in %s. Please, create it first or make sure the right path is being used."), *FilePath);
 		return nullptr;
 	}
 	// Return NeuralNetworkInferenceQAAsset
@@ -245,8 +260,10 @@ bool UNeuralNetworkInferenceQAAsset::Save()
 	FString FilePath;
 	{
 		FString ParentPath, FileName, FileExtension;
-		FPaths::Split(Package->GetLoadedPath().GetLocalFullPath(), ParentPath, FileName, FileExtension); // "../../../Sandbox/MachineLearning/NNIUnitTest/Content/NeuralNetworkInferenceQAAsset.uasset"
-		FilePath = ParentPath / UNeuralNetworkInferenceQAAsset_FileName + TEXT("_new.") + FileExtension; // "../../../Sandbox/MachineLearning/NNIUnitTest/Content/NeuralNetworkInferenceQAAsset_new.uasset"
+		// E.g., "../../../Sandbox/MachineLearning/NNIUnitTest/Content/NeuralNetworkInferenceQAAsset.uasset"
+		FPaths::Split(Package->GetLoadedPath().GetLocalFullPath(), ParentPath, FileName, FileExtension);
+		// E.g., "../../../Sandbox/MachineLearning/NNIUnitTest/Content/NeuralNetworkInferenceQAAsset_new.uasset"
+		FilePath = ParentPath / UNeuralNetworkInferenceQAAsset_FileName + TEXT("_new.") + FileExtension;
 	}
 	FSavePackageArgs SaveArgs;
 	SaveArgs.TopLevelFlags = RF_Standalone;
diff --git a/Engine/Plugins/Experimental/NeuralNetworkInference/Source/Private/NeuralNetworkInferenceQA/Private/NeuralNetworkInferenceQAUtils.cpp b/Engine/Plugins/Experimental/NeuralNetworkInference/Source/Private/NeuralNetworkInferenceQA/Private/NeuralNetworkInferenceQAUtils.cpp
--- a/Engine/Plugins/Experimental/NeuralNetworkInference/Source/Private/NeuralNetworkInferenceQA/Private/NeuralNetworkInferenceQAUtils.cpp
+++ b/Engine/Plugins/Experimental/NeuralNetworkInference/Source/Private/NeuralNetworkInferenceQA/Private/NeuralNetworkInferenceQAUtils.cpp
@@ -1,22 +1,26 @@
 // Copyright Epic Games, Inc. All Rights Reserved.
 
 #include "NeuralNetworkInferenceQAUtils.h"
+#include "NeuralTimer.h"
 
 
 
 /* FNeuralNetworkInferenceQAUtils static public functions
  *****************************************************************************/
 
-bool FNeuralNetworkInferenceQAUtils::EstimateTensorL1DiffError(const FNeuralTensor& InTensorA, const FNeuralTensor& InTensorB, const float InZeroThreshold, const FString& InDebugName)
+bool FNeuralNetworkInferenceQAUtils::EstimateTensorL1DiffError(const FNeuralTensor& InTensorA, const FNeuralTensor& InTensorB,
+	const float InZeroThreshold, const FString& InDebugName)
 {
 	if (InTensorA.GetDataType() != InTensorB.GetDataType())
 	{
-		UE_LOG(LogNeuralNetworkInferenceQA, Warning, TEXT("InTensorA.GetDataType() != InTensorB.GetDataType() failed, %d != %d."), InTensorA.GetDataType(), InTensorB.GetDataType());
+		UE_LOG(LogNeuralNetworkInferenceQA, Warning, TEXT("InTensorA.GetDataType() != InTensorB.GetDataType() failed, %d != %d."),
+			InTensorA.GetDataType(), InTensorB.GetDataType());
 		return false;
 	}
 	if (InTensorA.Num() != InTensorB.Num())
 	{
-		UE_LOG(LogNeuralNetworkInferenceQA, Warning, TEXT("InTensorA.Num() == InTensorB.Num() failed, %d != %d."), InTensorA.Num(), InTensorB.Num());
+		UE_LOG(LogNeuralNetworkInferenceQA, Warning, TEXT("InTensorA.Num() == InTensorB.Num() failed, %d != %d."), InTensorA.Num(),
+			InTensorB.Num());
 		return false;
 	}
 	const int32 TensorSize = InTensorA.Num();
@@ -83,7 +87,8 @@ bool FNeuralNetworkInferenceQAUtils::EstimateTensorL1DiffError(const FNeuralTens
 			{
 				Value1 = 0.f;
 				Value2 = 0.f;
-				ensureMsgf(false, TEXT("DataType = %d not implemented yet for FNeuralNetworkInferenceQAUtils::EstimateTensorL1DiffError()!"), InTensorA.GetDataType());
+				ensureMsgf(false, TEXT("DataType = %d not implemented yet for FNeuralNetworkInferenceQAUtils::EstimateTensorL1DiffError()!"),
+					InTensorA.GetDataType());
 			}
 			// Update L1
 			AverageL1Value += FMath::Abs(Value1);
@@ -100,7 +105,8 @@ bool FNeuralNetworkInferenceQAUtils::EstimateTensorL1DiffError(const FNeuralTens
 	// Some verbose
 	if (FiniteVsNon.Len() > 0)
 	{
-		UE_LOG(LogNeuralNetworkInferenceQA, Display, TEXT("IsFinite(Value1) != IsFinite(Value2) for %d values (skipped %d valid ones): %s"), TensorSize - SizeMinusInvalid, SizeMinusInvalid, *FiniteVsNon);
+		UE_LOG(LogNeuralNetworkInferenceQA, Display, TEXT("IsFinite(Value1) != IsFinite(Value2) for %d values (skipped %d valid ones): %s"),
+			TensorSize - SizeMinusInvalid, SizeMinusInvalid, *FiniteVsNon);
 	}
 	// Average L1
 	AverageL1Value /= SizeMinusInvalid;
@@ -108,13 +114,16 @@ bool FNeuralNetworkInferenceQAUtils::EstimateTensorL1DiffError(const FNeuralTens
 	// Sanity check
 	if (!FMath::IsFinite(AverageL1Value) || !FMath::IsFinite(AverageL1Diff))
 	{
-		UE_LOG(LogNeuralNetworkInferenceQA, Warning, TEXT("FNeuralNetworkInferenceQAUtils::EstimateTensorL1DiffError(): NaN|Inf found: AverageL1Diff %f for an L1 average value of %f."), AverageL1Diff, AverageL1Value);
+		UE_LOG(LogNeuralNetworkInferenceQA, Warning,
+			TEXT("FNeuralNetworkInferenceQAUtils::EstimateTensorL1DiffError(): NaN|Inf found: AverageL1Diff %f for an L1 average value of %f."),
+			AverageL1Diff, AverageL1Value);
 		return false;
 	}
 	// Did it pass?
 	if (AverageL1Diff > InZeroThreshold)
 	{
-		UE_LOG(LogNeuralNetworkInferenceQA, Warning, TEXT("%s: AverageL1Diff <= ZeroThreshold failed (%fe-3 > %fe-3), AverageL1Value = %f, AverageL1Diff/ZeroThreshold = %f."),
+		UE_LOG(LogNeuralNetworkInferenceQA, Warning,
+			TEXT("%s: AverageL1Diff <= ZeroThreshold failed (%fe-3 > %fe-3), AverageL1Value = %f, AverageL1Diff/ZeroThreshold = %f."),
 			*InDebugName, 1e3 * AverageL1Diff, 1e3 * InZeroThreshold, AverageL1Value, AverageL1Diff / InZeroThreshold);
 		UE_LOG(LogNeuralNetworkInferenceQA, Warning, TEXT("- TensorA = %s"), *InTensorA.ToString(100));
 		UE_LOG(LogNeuralNetworkInferenceQA, Warning, TEXT("- TensorB = %s"), *InTensorB.ToString(100));
@@ -122,3 +131,41 @@ bool FNeuralNetworkInferenceQAUtils::EstimateTensorL1DiffError(const FNeuralTens
 	}
 	return true;
 }
+
+TArray<FNeuralTensor> FNeuralNetworkInferenceQAUtils::CreateInputArrayCopy(const UNeuralNetwork* const InNeuralNetwork)
+{
+	TArray<FNeuralTensor> InputTensorArray;
+	for (uint32 InputTensorIndex = 0; InputTensorIndex < InNeuralNetwork->GetInputTensorNumber(); ++InputTensorIndex)
+	{
+		const FNeuralTensor& InputTensor = InNeuralNetwork->GetInputTensor(InputTensorIndex);
+		InputTensorArray.Push(InputTensor);
+	}
+	return InputTensorArray;
+}
+
+void FNeuralNetworkInferenceQAUtils::SetInputFromArrayCopy(UNeuralNetwork* InOutNeuralNetwork, const TArray<FNeuralTensor>& InInputTensorArray)
+{
+	// const FScopeLock ResourcesLock(&ResoucesCriticalSection);
+	if (InOutNeuralNetwork->GetInputTensorNumber() != InInputTensorArray.Num())
+	{
+		UE_LOG(LogNeuralNetworkInferenceQA, Warning,
+			TEXT("FNeuralNetworkInferenceQAUtils::SetInputFromArrayCopy(): GetInputTensorNumber() == InInputTensorArray.Num() failed, %d != %d."),
+			InOutNeuralNetwork->GetInputTensorNumber(), InInputTensorArray.Num());
+		return;
+	}
+	for (uint32 InputTensorIndex = 0; InputTensorIndex < InOutNeuralNetwork->GetInputTensorNumber(); ++InputTensorIndex)
+	{
+		InOutNeuralNetwork->SetInputFromVoidPointerCopy(InInputTensorArray[InputTensorIndex].GetUnderlyingUInt8ArrayRef().GetData(), InputTensorIndex);
+	}
+}
+
+TArray<FNeuralTensor> FNeuralNetworkInferenceQAUtils::CreateOutputArrayCopy(const UNeuralNetwork* const InNeuralNetwork)
+{
+	TArray<FNeuralTensor> OutputTensorArray;
+	for (uint32 OutputTensorIndex = 0; OutputTensorIndex < InNeuralNetwork->GetOutputTensorNumber(); ++OutputTensorIndex)
+	{
+		const FNeuralTensor& OutputTensor = InNeuralNetwork->GetOutputTensor(OutputTensorIndex);
+		OutputTensorArray.Push(OutputTensor);
+	}
+	return OutputTensorArray;
+}
diff --git a/Engine/Plugins/Experimental/NeuralNetworkInference/Source/Private/NeuralNetworkInferenceQA/Private/NeuralNetworkInferenceQAUtils.h b/Engine/Plugins/Experimental/NeuralNetworkInference/Source/Private/NeuralNetworkInferenceQA/Private/NeuralNetworkInferenceQAUtils.h
--- a/Engine/Plugins/Experimental/NeuralNetworkInference/Source/Private/NeuralNetworkInferenceQA/Private/NeuralNetworkInferenceQAUtils.h
+++ b/Engine/Plugins/Experimental/NeuralNetworkInference/Source/Private/NeuralNetworkInferenceQA/Private/NeuralNetworkInferenceQAUtils.h
@@ -3,6 +3,7 @@
 #pragma once
 
 #include "CoreMinimal.h"
+#include "NeuralNetwork.h"
 #include "NeuralTensor.h"
 
 DEFINE_LOG_CATEGORY_STATIC(LogNeuralNetworkInferenceQA, Display, All);
@@ -14,5 +15,15 @@ public:
 	 * It calculates the error between 2 data arrays (it should return 0 if both arrays are the same).
 	 * @return False if the estimated error > InZeroThreshold, true otherwise.
 	 */
-	static bool EstimateTensorL1DiffError(const FNeuralTensor& InTensorA, const FNeuralTensor& InTensorB, const float InZeroThreshold, const FString& InDebugName);
+	static bool EstimateTensorL1DiffError(const FNeuralTensor& InTensorA, const FNeuralTensor& InTensorB, const float InZeroThreshold,
+		const FString& InDebugName);
+
+	/**
+	 * Slow functions (as they will copy every input/output FNeuralNetwork) only meant for debugging purposes.
+	 * They really only copy the TArray<uint8> of underlying data, so they are confusing functions because they do not copy any other FNeuralTensor
+	 * properties.
+	 */
+	static TArray<FNeuralTensor> CreateInputArrayCopy(const UNeuralNetwork* const InNeuralNetwork);
+	static void SetInputFromArrayCopy(UNeuralNetwork* InOutNeuralNetwork, const TArray<FNeuralTensor>& InTensorDataArray);
+	static TArray<FNeuralTensor> CreateOutputArrayCopy(const UNeuralNetwork* const InNeuralNetwork);
 };
diff --git a/Engine/Plugins/Experimental/NeuralNetworkInference/Source/Private/NeuralNetworkInferenceQA/Private/OperatorTester.cpp b/Engine/Plugins/Experimental/NeuralNetworkInference/Source/Private/NeuralNetworkInferenceQA/Private/OperatorTester.cpp
--- a/Engine/Plugins/Experimental/NeuralNetworkInference/Source/Private/NeuralNetworkInferenceQA/Private/OperatorTester.cpp
+++ b/Engine/Plugins/Experimental/NeuralNetworkInference/Source/Private/NeuralNetworkInferenceQA/Private/OperatorTester.cpp
@@ -139,7 +139,8 @@ void FOperatorTester::TestBatchNormalizationOperator(UNeuralNetworkInferenceQAAs
 	UE_LOG(LogNeuralNetworkInferenceQA, Display, TEXT("--------------- %d/%d FBatchNormalizationOperator"), InOperatorCounter, InOperatorTotal);
 	// Inout tensors (X != Mean to avoid Output = Bias; Variance positive to avoid NaN)
 	TArray<TArray<FNeuralTensor>> Tensors({
-		// ----------------------------------------------------------------- ONNX - [REDACTED_PROFILE_URL]/onnx/blob/master/docs/TestCoverage.md#BatchNormalization -----------------------------------------------------------------
+		// ----------------------------------------------------------------- ONNX -----------------------------------------------------------------
+		// [REDACTED_PROFILE_URL]/onnx/blob/master/docs/TestCoverage.md#BatchNormalization
 		{FNeuralTensor(TArray<float>({-1,0,1, 2,3,4}), TArray<int64>({1,2,1,3})),
 			FNeuralTensor(TArray<float>({1, 1.5}), TArray<int64>({2})),
 			FNeuralTensor(TArray<float>({0, 1}), TArray<int64>({2})),
@@ -240,7 +241,8 @@ void FOperatorTester::TestGemmOperator(UNeuralNetworkInferenceQAAsset* InOutNetw
 	});
 	for (int32 TensorTestsIndex = 0; TensorTestsIndex < GemmABCOutputTensorTests.Num(); ++TensorTestsIndex)
 	{
-		UE_LOG(LogNeuralNetworkInferenceQA, Display, TEXT("-------------------- FGemmOperator #%d/%d"), TensorTestsIndex + 1, GemmABCOutputTensorTests.Num());
+		UE_LOG(LogNeuralNetworkInferenceQA, Display, TEXT("-------------------- FGemmOperator #%d/%d"),
+			TensorTestsIndex + 1, GemmABCOutputTensorTests.Num());
 		TArray<TSharedPtr<FNeuralOperator>> Operators({
 			// A, B
 			MakeShared<FGemmOperator>(/*Alpha*/-0.5, /*Beta*/-10, /*bTransA*/false, /*bTransB*/false),// = AB + C
@@ -317,7 +319,8 @@ void FOperatorTester::TestConvOperator(UNeuralNetworkInferenceQAAsset* InOutNetw
 	UE_LOG(LogNeuralNetworkInferenceQA, Display, TEXT("--------------- %d/%d FConvOperator"), InOperatorCounter, InOperatorTotal);
 	// Inout tensors
 	TArray<TArray<FNeuralTensor>> Tensors({
-		// ----------------------------------------------------------------- ONNX - [REDACTED_PROFILE_URL]/onnx/blob/master/docs/TestCoverage.md#Conv -----------------------------------------------------------------
+		// ----------------------------------------------------------------- ONNX -----------------------------------------------------------------
+		// [REDACTED_PROFILE_URL]/onnx/blob/master/docs/TestCoverage.md#Conv
 		// conv
 		{FNeuralTensor(TArray<float>({0,1,2,3,4, 5,6,7,8,9, 10,11,12,13,14, 15,16,17,18,19, 20,21,22,23,24}), TArray<int64>({1,1,5,5})),
 			FNeuralTensor(TArray<float>({1,1,1, 1,1,1, 1,1,1}), TArray<int64>({1,1,3,3})),
@@ -634,15 +637,16 @@ void FOperatorTester::TestConvOperator(UNeuralNetworkInferenceQAAsset* InOutNetw
 		// Set FConvBaseOperator::EAutoPad - Hacky but simple way of getting this TArray
 		// Create and push FConvOperator
 		Operators.Push(MakeShared<FConvOperator>(AutoPad, Dilations, Group, /*KernelShape*/TArray<int64>(), Pads, Strides));
-		// Remove all non-required/auxiliary tensors from InputTensors to keep it in X, W, and B (if B exists). Otherwise, those would be used by the network and moved into GPU memory
+		// Remove all non-required/auxiliary tensors from InputTensors to keep it in X, W, and B (if B exists). Otherwise, those would be used by
+		// the network and moved into GPU memory
 		while (Tensors[TensorIndex].Num() > 4 || Tensors[TensorIndex].Last().Num() < 1)
 		{
 			Tensors[TensorIndex].Pop();
 		}
 		// TensorsIndex is no longer valid after this while loop!
 	}
-	TestOperator(InOutNetworkInferenceQAAsset, Tensors, InputTensors, OutputTensors, Operators, TEXT("FConvOperator"), InGroundTruthDirectory, InZeroThreshold,
-		/*bShouldOperatorsProvideSameResult*/false, /*bInShouldRunGPU*/true);
+	TestOperator(InOutNetworkInferenceQAAsset, Tensors, InputTensors, OutputTensors, Operators, TEXT("FConvOperator"), InGroundTruthDirectory,
+		InZeroThreshold, /*bShouldOperatorsProvideSameResult*/false, /*bInShouldRunGPU*/true);
 }
 
 void FOperatorTester::TestConvTransposeOperator(UNeuralNetworkInferenceQAAsset* InOutNetworkInferenceQAAsset, const TMap<int32, TArray<float>>& InInputTensorData, const TMap<int32, TArray<float>>& InInputTensorDataSmall,
@@ -651,7 +655,8 @@ void FOperatorTester::TestConvTransposeOperator(UNeuralNetworkInferenceQAAsset*
 	UE_LOG(LogNeuralNetworkInferenceQA, Display, TEXT("--------------- %d/%d FConvTransposeOperator"), InOperatorCounter, InOperatorTotal);
 	// Inout tensors
 	TArray<TArray<FNeuralTensor>> Tensors({
-		// ----------------------------------------------------------------- ONNX - [REDACTED_PROFILE_URL]/onnx/blob/master/docs/TestCoverage.md#ConvTranspose -----------------------------------------------------------------
+		// ----------------------------------------------------------------- ONNX -----------------------------------------------------------------
+		// [REDACTED_PROFILE_URL]/onnx/blob/master/docs/TestCoverage.md#ConvTranspose
 		// convtranspose_1d
 		{FNeuralTensor(TArray<float>({0,1,2}), TArray<int64>({1,1,3})),
 			FNeuralTensor(TArray<float>({1,1,1, 1,1,1}), TArray<int64>({1,2,3})),
@@ -1017,36 +1022,49 @@ void FOperatorTester::TestReshapeOperator(UNeuralNetworkInferenceQAAsset* InOutN
 	UE_LOG(LogNeuralNetworkInferenceQA, Display, TEXT("--------------- %d/%d FReshapeOperator"), InOperatorCounter, InOperatorTotal);
 	// Inout tensors (X != Mean to avoid Output = Bias; Variance positive to avoid NaN)
 	TArray<TArray<FNeuralTensor>> Tensors({
-		// ----------------------------------------------------------------- ONNX - [REDACTED_PROFILE_URL]/onnx/blob/master/docs/TestCoverage.md#Reshape -----------------------------------------------------------------
+		// ----------------------------------------------------------------- ONNX -----------------------------------------------------------------
+		// [REDACTED_PROFILE_URL]/onnx/blob/master/docs/TestCoverage.md#Reshape
 		// reshape
-		{FNeuralTensor(InInputTensorData.FindChecked(24), TArray<int64>({2, 3, 4})), FNeuralTensor(TArray<int64>({4, 2, 3})), /*Output*/FNeuralTensor()},
+		{FNeuralTensor(InInputTensorData.FindChecked(24), TArray<int64>({2, 3, 4})), FNeuralTensor(TArray<int64>({4, 2, 3})),
+			/*Output*/FNeuralTensor()},
 		{FNeuralTensor(InInputTensorData.FindChecked(24), TArray<int64>({2, 3, 4})), FNeuralTensor(TArray<int64>({2, 4, 3}))},
-		{FNeuralTensor(InInputTensorData.FindChecked(24), TArray<int64>({2, 3, 4})), FNeuralTensor(TArray<int64>({2, 12})), /*Output*/FNeuralTensor()},
+		{FNeuralTensor(InInputTensorData.FindChecked(24), TArray<int64>({2, 3, 4})), FNeuralTensor(TArray<int64>({2, 12})),
+			/*Output*/FNeuralTensor()},
 		{FNeuralTensor(InInputTensorData.FindChecked(24), TArray<int64>({2, 3, 4})), FNeuralTensor(TArray<int64>({2, 3, 2, 2}))},
-		{FNeuralTensor(InInputTensorData.FindChecked(24), TArray<int64>({2, 3, 4})), FNeuralTensor(TArray<int64>({24})), /*Output*/FNeuralTensor()},
+		{FNeuralTensor(InInputTensorData.FindChecked(24), TArray<int64>({2, 3, 4})), FNeuralTensor(TArray<int64>({24})),
+			/*Output*/FNeuralTensor()},
 		{FNeuralTensor(InInputTensorData.FindChecked(24), TArray<int64>({2, 3, 4})), FNeuralTensor(TArray<int64>({2, -1, 2}))},
-		{FNeuralTensor(InInputTensorData.FindChecked(24), TArray<int64>({2, 3, 4})), FNeuralTensor(TArray<int64>({-1, 2, 3, 4})), /*Output*/FNeuralTensor()},
+		{FNeuralTensor(InInputTensorData.FindChecked(24), TArray<int64>({2, 3, 4})), FNeuralTensor(TArray<int64>({-1, 2, 3, 4})),
+			/*Output*/FNeuralTensor()},
 		{FNeuralTensor(InInputTensorData.FindChecked(24), TArray<int64>({2, 3, 4})), FNeuralTensor(TArray<int64>({2, 0, 4, 1}))},
-		{FNeuralTensor(InInputTensorData.FindChecked(24), TArray<int64>({2, 3, 4})), FNeuralTensor(TArray<int64>({2, 0, 1, -1})), /*Output*/FNeuralTensor()},
+		{FNeuralTensor(InInputTensorData.FindChecked(24), TArray<int64>({2, 3, 4})), FNeuralTensor(TArray<int64>({2, 0, 1, -1})),
+			/*Output*/FNeuralTensor()},
 		// allowzero
-		{FNeuralTensor(FNeuralTensor(ENeuralDataType::Float, TArray<int64>({0, 3, 4}))), FNeuralTensor(TArray<int64>({3, 4, 0})), /*Output*/FNeuralTensor(), /*AllowZero=true*/FNeuralTensor()},
+		{FNeuralTensor(FNeuralTensor(ENeuralDataType::Float, TArray<int64>({0, 3, 4}))), FNeuralTensor(TArray<int64>({3, 4, 0})),
+			/*Output*/FNeuralTensor(), /*AllowZero=true*/FNeuralTensor()},
 		// Scalars
-		{FNeuralTensor(InInputTensorData.FindChecked(1)), FNeuralTensor(TArray<int64>({1})), /*Output*/FNeuralTensor()},
+		{FNeuralTensor(InInputTensorData.FindChecked(1)), FNeuralTensor(TArray<int64>({1})),
+			/*Output*/FNeuralTensor()},
 		{FNeuralTensor(InInputTensorData.FindChecked(1)), FNeuralTensor(TArray<int64>({1}))},
 		// 1-D
-		{FNeuralTensor(InInputTensorData.FindChecked(10)), FNeuralTensor(TArray<int64>({5, 2})), /*Output*/FNeuralTensor()},
+		{FNeuralTensor(InInputTensorData.FindChecked(10)), FNeuralTensor(TArray<int64>({5, 2})),
+			/*Output*/FNeuralTensor()},
 		{FNeuralTensor(InInputTensorData.FindChecked(10)), FNeuralTensor(TArray<int64>({5, -1}))},
 		// 2-D
-		{FNeuralTensor(InInputTensorData.FindChecked(18), TArray<int64>({3, 6})), FNeuralTensor(TArray<int64>({-1, 3})), /*Output*/FNeuralTensor()},
+		{FNeuralTensor(InInputTensorData.FindChecked(18), TArray<int64>({3, 6})), FNeuralTensor(TArray<int64>({-1, 3})),
+			/*Output*/FNeuralTensor()},
 		{FNeuralTensor(InInputTensorData.FindChecked(18), TArray<int64>({3, 6})), FNeuralTensor(TArray<int64>({6, 3}))},
 		// 3-D
-		{FNeuralTensor(InInputTensorData.FindChecked(30), TArray<int64>({3, 5, 2})), FNeuralTensor(TArray<int64>({30})), /*Output*/FNeuralTensor()},
+		{FNeuralTensor(InInputTensorData.FindChecked(30), TArray<int64>({3, 5, 2})), FNeuralTensor(TArray<int64>({30})),
+			/*Output*/FNeuralTensor()},
 		{FNeuralTensor(InInputTensorData.FindChecked(30), TArray<int64>({3, 5, 2})), FNeuralTensor(TArray<int64>({-1}))},
 		// 4-D
-		{FNeuralTensor(InInputTensorData.FindChecked(300), TArray<int64>({3, 10, 5, 2})), FNeuralTensor(TArray<int64>({-1, 10, 15})), /*Output*/FNeuralTensor()},
+		{FNeuralTensor(InInputTensorData.FindChecked(300), TArray<int64>({3, 10, 5, 2})), FNeuralTensor(TArray<int64>({-1, 10, 15})),
+			/*Output*/FNeuralTensor()},
 		{FNeuralTensor(InInputTensorData.FindChecked(300), TArray<int64>({3, 10, 5, 2})), FNeuralTensor(TArray<int64>({2, 0, 15}))},
 		// 5-D
-		{FNeuralTensor(InInputTensorData.FindChecked(300), TArray<int64>({2, 3, 5, 5, 2})), FNeuralTensor(TArray<int64>({2, 10, 15})), /*Output*/FNeuralTensor()},
+		{FNeuralTensor(InInputTensorData.FindChecked(300), TArray<int64>({2, 3, 5, 5, 2})), FNeuralTensor(TArray<int64>({2, 10, 15})),
+			/*Output*/FNeuralTensor()},
 		{FNeuralTensor(InInputTensorData.FindChecked(300), TArray<int64>({2, 3, 5, 5, 2})), FNeuralTensor(TArray<int64>({0, -1, 15}))}
 		});
 	TArray<TArray<FNeuralTensor*>> InputTensors;
@@ -1068,49 +1086,68 @@ void FOperatorTester::TestReshapeOperator(UNeuralNetworkInferenceQAAsset* InOutN
 			Operators.Push(MakeShared<FReshapeOperator>(/*bIsInlinedTensor*/true, /*AllowZero*/0));
 		}
 	}
-	TestOperator(InOutNetworkInferenceQAAsset, Tensors, InputTensors, OutputTensors, Operators, TEXT("FReshapeOperator"), InGroundTruthDirectory, InZeroThreshold,
-		/*bShouldOperatorsProvideSameResult*/false, /*bInShouldRunGPU*/true);
+	TestOperator(InOutNetworkInferenceQAAsset, Tensors, InputTensors, OutputTensors, Operators, TEXT("FReshapeOperator"), InGroundTruthDirectory,
+		InZeroThreshold, /*bShouldOperatorsProvideSameResult*/false, /*bInShouldRunGPU*/true);
 }
 
-void FOperatorTester::TestSqueezeOperator(UNeuralNetworkInferenceQAAsset* InOutNetworkInferenceQAAsset, const TMap<int32, TArray<float>>& InInputTensorData,
-	const TMap<int32, TArray<float>>& InInputTensorDataPos, const float InZeroThreshold, const FString& InGroundTruthDirectory, const int32 InOperatorCounter, const int32 InOperatorTotal)
+void FOperatorTester::TestSqueezeOperator(UNeuralNetworkInferenceQAAsset* InOutNetworkInferenceQAAsset,
+	const TMap<int32, TArray<float>>& InInputTensorData, const TMap<int32, TArray<float>>& InInputTensorDataPos, const float InZeroThreshold,
+	const FString& InGroundTruthDirectory, const int32 InOperatorCounter, const int32 InOperatorTotal)
 {
 	UE_LOG(LogNeuralNetworkInferenceQA, Display, TEXT("--------------- %d/%d FSqueezeOperator"), InOperatorCounter, InOperatorTotal);
 	// Inout tensors (X != Mean to avoid Output = Bias; Variance positive to avoid NaN)
 	TArray<TArray<FNeuralTensor>> Tensors({
-		// ----------------------------------------------------------------- ONNX - [REDACTED_PROFILE_URL]/onnx/blob/master/docs/Operators.md#Squeeze -----------------------------------------------------------------
+		// ----------------------------------------------------------------- ONNX -----------------------------------------------------------------
+		// [REDACTED_PROFILE_URL]/onnx/blob/master/docs/Operators.md#Squeeze
 		// squeeze
-		{FNeuralTensor(InInputTensorData.FindChecked(60), TArray<int64>({1, 3, 4, 5})), FNeuralTensor(TArray<int64>({0})), /*Output*/FNeuralTensor()},
+		{FNeuralTensor(InInputTensorData.FindChecked(60), TArray<int64>({1, 3, 4, 5})), FNeuralTensor(TArray<int64>({0})),
+			/*Output*/FNeuralTensor()},
 		{FNeuralTensor(InInputTensorData.FindChecked(60), TArray<int64>({1, 3, 4, 5})), FNeuralTensor(TArray<int64>({0}))},
 		// squeeze_negative_axis
-		{FNeuralTensor(InInputTensorData.FindChecked(15), TArray<int64>({1, 3, 1, 5})), FNeuralTensor(TArray<int64>({-2})), /*Output*/FNeuralTensor()},
+		{FNeuralTensor(InInputTensorData.FindChecked(15), TArray<int64>({1, 3, 1, 5})), FNeuralTensor(TArray<int64>({-2})),
+			/*Output*/FNeuralTensor()},
 		{FNeuralTensor(InInputTensorData.FindChecked(15), TArray<int64>({1, 3, 1, 5})), FNeuralTensor(TArray<int64>({-2}))},
 		// 1-D
 		{FNeuralTensor(InInputTensorData.FindChecked(5)), /*Empty*/FNeuralTensor(), /*Empty*/FNeuralTensor(), /*Output*/FNeuralTensor()},
 		// 2-D
-		{FNeuralTensor(InInputTensorData.FindChecked(60), TArray<int64>({12, 5})), /*Empty*/FNeuralTensor(), /*Empty*/FNeuralTensor(), /*Output*/FNeuralTensor()},
-		{FNeuralTensor(InInputTensorData.FindChecked(6), TArray<int64>({6, 1})), FNeuralTensor(TArray<int64>({1})), /*Output*/FNeuralTensor()},
+		{FNeuralTensor(InInputTensorData.FindChecked(60), TArray<int64>({12, 5})), /*Empty*/FNeuralTensor(), /*Empty*/FNeuralTensor(),
+			/*Output*/FNeuralTensor()},
+		{FNeuralTensor(InInputTensorData.FindChecked(6), TArray<int64>({6, 1})), FNeuralTensor(TArray<int64>({1})),
+			/*Output*/FNeuralTensor()},
 		{FNeuralTensor(InInputTensorData.FindChecked(6), TArray<int64>({6, 1})), FNeuralTensor(TArray<int64>({1}))},
 		{FNeuralTensor(InInputTensorData.FindChecked(6), TArray<int64>({6, 1}))},
-		{FNeuralTensor(InInputTensorData.FindChecked(6), TArray<int64>({6, 1})), /*Empty*/FNeuralTensor(), /*Empty*/FNeuralTensor(), /*Output*/FNeuralTensor()},
+		{FNeuralTensor(InInputTensorData.FindChecked(6), TArray<int64>({6, 1})), /*Empty*/FNeuralTensor(), /*Empty*/FNeuralTensor(),
+			/*Output*/FNeuralTensor()},
 		// 3-D
-		{FNeuralTensor(InInputTensorData.FindChecked(60), TArray<int64>({3, 4, 5})), /*Empty*/FNeuralTensor(), /*Empty*/FNeuralTensor(), /*Output*/FNeuralTensor()},
+		{FNeuralTensor(InInputTensorData.FindChecked(60), TArray<int64>({3, 4, 5})), /*Empty*/FNeuralTensor(), /*Empty*/FNeuralTensor(),
+			/*Output*/FNeuralTensor()},
 		{FNeuralTensor(InInputTensorData.FindChecked(30), TArray<int64>({5, 1, 6})), FNeuralTensor(TArray<int64>({1}))},
-		{FNeuralTensor(InInputTensorData.FindChecked(30), TArray<int64>({5, 1, 6})), FNeuralTensor(TArray<int64>({-2})), /*Output*/FNeuralTensor()},
-		{FNeuralTensor(InInputTensorData.FindChecked(30), TArray<int64>({5, 1, 6})), FNeuralTensor(TArray<int64>({1})), /*Output*/FNeuralTensor()},
-		{FNeuralTensor(InInputTensorData.FindChecked(30), TArray<int64>({5, 1, 6})), /*Empty*/FNeuralTensor(), /*Empty*/FNeuralTensor(), /*Output*/FNeuralTensor()},
+		{FNeuralTensor(InInputTensorData.FindChecked(30), TArray<int64>({5, 1, 6})), FNeuralTensor(TArray<int64>({-2})),
+			/*Output*/FNeuralTensor()},
+		{FNeuralTensor(InInputTensorData.FindChecked(30), TArray<int64>({5, 1, 6})), FNeuralTensor(TArray<int64>({1})),
+			/*Output*/FNeuralTensor()},
+		{FNeuralTensor(InInputTensorData.FindChecked(30), TArray<int64>({5, 1, 6})), /*Empty*/FNeuralTensor(), /*Empty*/FNeuralTensor(),
+			/*Output*/FNeuralTensor()},
 		{FNeuralTensor(InInputTensorData.FindChecked(30), TArray<int64>({5, 1, 6}))},
 		// 4-D
-		{FNeuralTensor(InInputTensorData.FindChecked(60), TArray<int64>({3, 2, 2, 5})), /*Empty*/FNeuralTensor(), /*Empty*/FNeuralTensor(), /*Output*/FNeuralTensor()},
-		{FNeuralTensor(InInputTensorData.FindChecked(16), TArray<int64>({2, 8, 1, 1})), FNeuralTensor(TArray<int64>({-1, 2})), /*Output*/FNeuralTensor()},
-		{FNeuralTensor(InInputTensorData.FindChecked(16), TArray<int64>({2, 8, 1, 1})), /*Empty*/FNeuralTensor(), /*Empty*/FNeuralTensor(), /*Output*/FNeuralTensor()},
+		{FNeuralTensor(InInputTensorData.FindChecked(60), TArray<int64>({3, 2, 2, 5})), /*Empty*/FNeuralTensor(), /*Empty*/FNeuralTensor(),
+			/*Output*/FNeuralTensor()},
+		{FNeuralTensor(InInputTensorData.FindChecked(16), TArray<int64>({2, 8, 1, 1})), FNeuralTensor(TArray<int64>({-1, 2})),
+			/*Output*/FNeuralTensor()},
+		{FNeuralTensor(InInputTensorData.FindChecked(16), TArray<int64>({2, 8, 1, 1})), /*Empty*/FNeuralTensor(), /*Empty*/FNeuralTensor(),
+			/*Output*/FNeuralTensor()},
 		{FNeuralTensor(InInputTensorData.FindChecked(16), TArray<int64>({2, 8, 1, 1}))},
-		{FNeuralTensor(InInputTensorData.FindChecked(16), TArray<int64>({2, 8, 1, 1})), FNeuralTensor(TArray<int64>({2})), /*Output*/FNeuralTensor()},
+		{FNeuralTensor(InInputTensorData.FindChecked(16), TArray<int64>({2, 8, 1, 1})), FNeuralTensor(TArray<int64>({2})),
+			/*Output*/FNeuralTensor()},
 		// 5-D
-		{FNeuralTensor(InInputTensorData.FindChecked(50), TArray<int64>({1, 5, 1, 5, 2})), FNeuralTensor(TArray<int64>({0, 2})), /*Output*/FNeuralTensor()},
-		{FNeuralTensor(InInputTensorData.FindChecked(50), TArray<int64>({1, 5, 1, 5, 2})), FNeuralTensor(TArray<int64>({-3, -5})), /*Output*/FNeuralTensor()},
-		{FNeuralTensor(InInputTensorData.FindChecked(27), TArray<int64>({3, 1, 3, 1, 3})), FNeuralTensor(TArray<int64>({1, 3})), /*Output*/FNeuralTensor()},
-		{FNeuralTensor(InInputTensorData.FindChecked(27), TArray<int64>({3, 1, 3, 1, 3})), FNeuralTensor(TArray<int64>({3, 1})), /*Output*/FNeuralTensor()},
+		{FNeuralTensor(InInputTensorData.FindChecked(50), TArray<int64>({1, 5, 1, 5, 2})), FNeuralTensor(TArray<int64>({0, 2})),
+			/*Output*/FNeuralTensor()},
+		{FNeuralTensor(InInputTensorData.FindChecked(50), TArray<int64>({1, 5, 1, 5, 2})), FNeuralTensor(TArray<int64>({-3, -5})),
+			/*Output*/FNeuralTensor()},
+		{FNeuralTensor(InInputTensorData.FindChecked(27), TArray<int64>({3, 1, 3, 1, 3})), FNeuralTensor(TArray<int64>({1, 3})),
+			/*Output*/FNeuralTensor()},
+		{FNeuralTensor(InInputTensorData.FindChecked(27), TArray<int64>({3, 1, 3, 1, 3})), FNeuralTensor(TArray<int64>({3, 1})),
+			/*Output*/FNeuralTensor()},
 		{FNeuralTensor(InInputTensorData.FindChecked(50), TArray<int64>({1, 5, 1, 5, 2}))}
 		//// Other tests (resulting in Warnings and/or unexpected behavior)
 		//// Scalars (undefined behavior: 0-d shape array for CPU and crash for GPU)
@@ -1119,10 +1156,14 @@ void FOperatorTester::TestSqueezeOperator(UNeuralNetworkInferenceQAAsset* InOutN
 		//{FNeuralTensor(InInputTensorData.FindChecked(1)), /*Output*/FNeuralTensor()},
 		//{FNeuralTensor(InInputTensorData.FindChecked(1))},
 		//// Duplicated axis (generating warnings)
-		//{FNeuralTensor(InInputTensorData.FindChecked(16), TArray<int64>({2, 8, 1, 1})), FNeuralTensor(TArray<int64>({-2, 2})),  /*Output*/FNeuralTensor()},
-		//{FNeuralTensor(InInputTensorData.FindChecked(27), TArray<int64>({3, 1, 3, 1, 3})), FNeuralTensor(TArray<int64>({1, 1})), /*Output*/FNeuralTensor()},
-		//{FNeuralTensor(InInputTensorData.FindChecked(27), TArray<int64>({3, 1, 3, 1, 3})), FNeuralTensor(TArray<int64>({-4, 1})), /*Output*/FNeuralTensor()},
-		//{FNeuralTensor(InInputTensorData.FindChecked(27), TArray<int64>({3, 1, 3, 1, 3})), FNeuralTensor(TArray<int64>({1, -4})), /*Output*/FNeuralTensor()}
+		//{FNeuralTensor(InInputTensorData.FindChecked(16), TArray<int64>({2, 8, 1, 1})), FNeuralTensor(TArray<int64>({-2, 2})),
+		//	/*Output*/FNeuralTensor()},
+		//{FNeuralTensor(InInputTensorData.FindChecked(27), TArray<int64>({3, 1, 3, 1, 3})), FNeuralTensor(TArray<int64>({1, 1})),
+		//	/*Output*/FNeuralTensor()},
+		//{FNeuralTensor(InInputTensorData.FindChecked(27), TArray<int64>({3, 1, 3, 1, 3})), FNeuralTensor(TArray<int64>({-4, 1})),
+		//	/*Output*/FNeuralTensor()},
+		//{FNeuralTensor(InInputTensorData.FindChecked(27), TArray<int64>({3, 1, 3, 1, 3})), FNeuralTensor(TArray<int64>({1, -4})),
+		//	/*Output*/FNeuralTensor()}
 		});
 	TArray<TArray<FNeuralTensor*>> InputTensors;
 	TArray<TArray<FNeuralTensor*>> OutputTensors;
@@ -1159,8 +1200,8 @@ void FOperatorTester::TestSqueezeOperator(UNeuralNetworkInferenceQAAsset* InOutN
 			Operators.Push(MakeShared<FSqueezeOperator>(/*bIsInlinedTensor*/true));
 		}
 	}
-	TestOperator(InOutNetworkInferenceQAAsset, Tensors, InputTensors, OutputTensors, Operators, TEXT("FSqueezeOperator"), InGroundTruthDirectory, InZeroThreshold,
-		/*bShouldOperatorsProvideSameResult*/false, /*bInShouldRunGPU*/true);
+	TestOperator(InOutNetworkInferenceQAAsset, Tensors, InputTensors, OutputTensors, Operators, TEXT("FSqueezeOperator"), InGroundTruthDirectory,
+		InZeroThreshold, /*bShouldOperatorsProvideSameResult*/false, /*bInShouldRunGPU*/true);
 }
 
 /* FOperatorTester static private functions
@@ -1186,7 +1227,7 @@ void FOperatorTester::ResetOrCheckInput(UNeuralNetwork* InOutNetwork, const TArr
 {
 	if (bIsInlinedTensor)
 	{
-		InOutNetwork->SetInputFromArrayCopy(InInputTensorArray);
+		FNeuralNetworkInferenceQAUtils::SetInputFromArrayCopy(InOutNetwork, InInputTensorArray);
 	}
 	// Check input has not changed
 	else
@@ -1194,9 +1235,11 @@ void FOperatorTester::ResetOrCheckInput(UNeuralNetwork* InOutNetwork, const TArr
 		for (int32 InputIndex = 0; InputIndex < InOutNetwork->GetInputTensorNumber(); ++InputIndex)
 		{
 			const FNeuralTensor& InputTensor = InOutNetwork->GetInputTensor(InputIndex);
-			ensureMsgf(InInputTensorArray[InputIndex].Num() == InputTensor.Num(), TEXT("InInputTensorArray[InputIndex].Num() == InputTensor.Num() failed (%d != %d)."),
+			ensureMsgf(InInputTensorArray[InputIndex].Num() == InputTensor.Num(),
+				TEXT("InInputTensorArray[InputIndex].Num() == InputTensor.Num() failed (%d != %d)."),
 				InInputTensorArray[InputIndex].Num(), InputTensor.Num());
-			ensureMsgf(FNeuralNetworkInferenceQAUtils::EstimateTensorL1DiffError(InputTensor, InInputTensorArray[InputIndex], 0.f, TEXT("ResetOrCheckInput")), TEXT("ResetOrCheckInput() did not pass."));
+			ensureMsgf(FNeuralNetworkInferenceQAUtils::EstimateTensorL1DiffError(InputTensor, InInputTensorArray[InputIndex], 0.f,
+				TEXT("ResetOrCheckInput")), TEXT("ResetOrCheckInput() did not pass."));
 		}
 	}
 }
@@ -1204,12 +1247,12 @@ void FOperatorTester::ResetOrCheckInput(UNeuralNetwork* InOutNetwork, const TArr
 void FOperatorTester::NetworkGPUvsCPU(UNeuralNetwork* InOutNetwork, const bool bIsInlinedTensor, const float InZeroThreshold)
 {
 	// Create InputTensorDataArray
-	const TArray<FNeuralTensor> InputTensorArray = InOutNetwork->CreateInputArrayCopy();
+	const TArray<FNeuralTensor> InputTensorArray = FNeuralNetworkInferenceQAUtils::CreateInputArrayCopy(InOutNetwork);
 	// GPU
 	InOutNetwork->SetDeviceType(ENeuralDeviceType::GPU);
 	InOutNetwork->Run();
 	// Store GPU results
-	const TArray<FNeuralTensor> OutputTensorArray = InOutNetwork->CreateOutputArrayCopy();
+	const TArray<FNeuralTensor> OutputTensorArray = FNeuralNetworkInferenceQAUtils::CreateOutputArrayCopy(InOutNetwork);
 	// Reset memory / Check input has not changed
 	ResetOrCheckInput(InOutNetwork, InputTensorArray, bIsInlinedTensor);
 	// CPU
@@ -1222,7 +1265,8 @@ void FOperatorTester::NetworkGPUvsCPU(UNeuralNetwork* InOutNetwork, const bool b
 		const FNeuralTensor& OutputTensorGPU = InOutNetwork->GetOutputTensor(OutputIndex);
 		ensureMsgf(OutputTensorCPU.Num() == OutputTensorGPU.Num(), TEXT("OutputTensorCPU.Num() == OutputTensorGPU.Num(): %d vs. %d"),
 			OutputTensorCPU.Num(), OutputTensorGPU.Num());
-		ensureMsgf(FNeuralNetworkInferenceQAUtils::EstimateTensorL1DiffError(OutputTensorCPU, OutputTensorGPU, InZeroThreshold, TEXT("GPUvsCPU")), TEXT("CPU vs. GPU error has too high"));
+		ensureMsgf(FNeuralNetworkInferenceQAUtils::EstimateTensorL1DiffError(OutputTensorCPU, OutputTensorGPU, InZeroThreshold, TEXT("GPUvsCPU")),
+			TEXT("CPU vs. GPU error has too high"));
 	}
 	// Reset memory / Check input has not changed
 	ResetOrCheckInput(InOutNetwork, InputTensorArray, bIsInlinedTensor); // This is doing an unnecessary copy (it should move really)
@@ -1230,7 +1274,7 @@ void FOperatorTester::NetworkGPUvsCPU(UNeuralNetwork* InOutNetwork, const bool b
 
 void FOperatorTester::NetworkInputVsOutput(UNeuralNetwork* InOutNetwork, const int32 InInlinedTensorIndex)
 {
-	const TArray<FNeuralTensor> InputTensorArray = InOutNetwork->CreateInputArrayCopy();
+	const TArray<FNeuralTensor> InputTensorArray = FNeuralNetworkInferenceQAUtils::CreateInputArrayCopy(InOutNetwork);
 	// Run UNeuralNetwork
 	InOutNetwork->Run();
 	// Check input equals one of the outputs
@@ -1247,5 +1291,5 @@ void FOperatorTester::NetworkInputVsOutput(UNeuralNetwork* InOutNetwork, const i
 	}
 	ensureMsgf(bIsSomeInputAnOutput, TEXT("Input is not being used as an output."));
 	// Reset memory
-	InOutNetwork->SetInputFromArrayCopy(InputTensorArray); // This is doing an unnecessary copy (it should move really)
+	FNeuralNetworkInferenceQAUtils::SetInputFromArrayCopy(InOutNetwork, InputTensorArray); // This is doing an unnecessary copy (it should move)
 }
