Added missing include.

#rnx

diff --git a/Engine/Plugins/Experimental/NeuralNetworkInference/Source/Private/ThirdPartyHelperAndDLLLoader/Private/Module.cpp b/Engine/Plugins/Experimental/NeuralNetworkInference/Source/Private/ThirdPartyHelperAndDLLLoader/Private/Module.cpp
--- a/Engine/Plugins/Experimental/NeuralNetworkInference/Source/Private/ThirdPartyHelperAndDLLLoader/Private/Module.cpp
+++ b/Engine/Plugins/Experimental/NeuralNetworkInference/Source/Private/ThirdPartyHelperAndDLLLoader/Private/Module.cpp
@@ -1,6 +1,7 @@
 // Copyright Epic Games, Inc. All Rights Reserved.
 
 #include "Module.h"
+#include "HAL/FileManager.h"
 #include "HAL/PlatformProcess.h"
 #include "Interfaces/IPluginManager.h"
 #include "Misc/Paths.h"
