More cleaned / commented UNN

#rnx

diff --git a/Engine/Plugins/Experimental/NeuralNetworkInference/Source/NeuralNetworkInference/Private/NeuralNetwork.cpp b/Engine/Plugins/Experimental/NeuralNetworkInference/Source/NeuralNetworkInference/Private/NeuralNetwork.cpp
--- a/Engine/Plugins/Experimental/NeuralNetworkInference/Source/NeuralNetworkInference/Private/NeuralNetwork.cpp
+++ b/Engine/Plugins/Experimental/NeuralNetworkInference/Source/NeuralNetworkInference/Private/NeuralNetwork.cpp
@@ -506,7 +506,7 @@ void UNeuralNetwork::Run()
 	ComputeStatsModule.StoreSample(RunTimer.Toc());
 }
 
-float UNeuralNetwork::GetLastInferenceTime() const
+float UNeuralNetwork::GetLastInferenceTimeMSec() const
 {
 	return ComputeStatsModule.GetLastSample();
 }
diff --git a/Engine/Plugins/Experimental/NeuralNetworkInference/Source/NeuralNetworkInference/Public/NeuralNetwork.h b/Engine/Plugins/Experimental/NeuralNetworkInference/Source/NeuralNetworkInference/Public/NeuralNetwork.h
--- a/Engine/Plugins/Experimental/NeuralNetworkInference/Source/NeuralNetworkInference/Public/NeuralNetwork.h
+++ b/Engine/Plugins/Experimental/NeuralNetworkInference/Source/NeuralNetworkInference/Public/NeuralNetwork.h
@@ -48,7 +48,7 @@
  * 2.a. Running inference:
  *		// Fill input neural tensor
  *		TArray<float> InArray; // Assumed initialized with data and that InArray.Num() == Network->Num()
- *		Network->SetInputFromArrayCopy(InArray);
+ *		Network->SetInputFromArrayCopy(InArray); // Equivalent: Network->SetInputFromVoidPointerCopy(InArray.GetData());
  *		UE_LOG(LogNeuralNetworkInference, Display, TEXT("Input tensor: %s."), *Network->GetInputTensor().ToString());
  *		// Run UNeuralNetwork
  *		Network->Run();
@@ -64,10 +64,10 @@
  *			InputDataPointer[Index] = ...; // Assumed some preprocessing or otherwise simply use Memcpy
  *
  * 3. Networks with multiple input/output tensors:
- *      - Multiple inputs: Add InTensorIndex to SetInputFromArrayCopy, GetInputTensor(InTensorIndex) or GetInputDataPointerMutable(InTensorIndex) in
- *        the examples above:
- *		    Network->SetInputFromArrayCopy(InputArray0, 0);
- *		    Network->SetInputFromArrayCopy(InputArray1, 1);
+ *      - Multiple inputs: Add InTensorIndex to SetInputFromArrayCopy, SetInputFromVoidPointerCopy, GetInputTensor or GetInputDataPointerMutable
+ *        in the examples above:
+ *		    Network->SetInputFromArrayCopy(InputArray0, 0); // Equivalent: Network->SetInputFromVoidPointerCopy(InputArray0.GetData(), 0);
+ *		    Network->SetInputFromArrayCopy(InputArray1, 1); // Equivalent: Network->SetInputFromVoidPointerCopy(InputArray1.GetData(), 1);
  *      - Multiple outputs: Add InTensorIndex to GetOutputTensor(InTensorIndex) in the examples above or use GetOutputTensors() instead.
  *		    const FNeuralTensor& OutputTensor0 = Network->GetOutputTensor(0);
  *		    const FNeuralTensor& OutputTensor1 = Network->GetOutputTensor(1);
@@ -89,6 +89,10 @@ public:
 	virtual ~UNeuralNetwork() = default;
 
 	/**
+	 * Load() + SetInputFromArrayCopy() + Run() is the simplest way to load an ONNX file, set the input tensor(s), and run inference on it. All other
+	 * functions in UNeuralNetwork provide additional functionality (set the device to CPU/GPU, more complicated but faster ways to set the input,
+	 * support for asynchronous run, avoid copying the memory CPU-to-GPU or GPU-to-CPU, etc.).
+	 *
 	 * It loads the desired network graph definition and weights from an input ONNX file. This file can be passed either a file path or as a memory
 	 * buffer.
 	 * @param InModelFilePath Input ONNX file path to be read. It can either be a full path or a relative path with respect to the Engine or Game
@@ -96,6 +100,8 @@ public:
 	 * @param InModelReadFromFileInBytes TArray buffer filled with the contents of the ONNX file. This TArray<uint8> buffer will be moved for
 	 * performance reasons. @see FFileHelper::LoadFileToArray for an example of how to read a file into a TArray<uint8>.
 	 * @return Whether the network was successfully loaded. Equivalent to IsLoaded(). @see IsLoaded() for more details.
+	 *
+	 * @see Load(), SetInputFromArrayCopy(), Run().
 	 */
 	bool Load(const FString& InModelFilePath);
 	bool Load(TArray<uint8>& InModelReadFromFileInBytes);
@@ -173,7 +179,8 @@ public:
 
 	/**
 	 * GetInputTensorNumber() and GetOutputTensorNumber() return the number of input or output tensors of this network, respectively.
-	 * @see GetInputTensor(), SetInputFromArrayCopy(), and GetInputDataPointerMutable() for other input-tensor-related functions.
+	 * @see GetInputTensor(), SetInputFromArrayCopy(), SetInputFromVoidPointerCopy(), and GetInputDataPointerMutable() for other
+	 * input-tensor-related functions.
 	 * @see GetOutputTensor() for other output-tensor-related functions.
 	 */
 	int64 GetInputTensorNumber() const;
@@ -189,21 +196,19 @@ public:
 	const FNeuralTensor& GetOutputTensor(const int32 InTensorIndex = 0) const;
 
 	/**
-	 * SetInputFromArrayCopy() and GetInputDataPointerMutable() are the only functions that allow modifing the network input tensor(s) values:
-	 * - SetInputFromArrayCopy() is very easy to use but less efficient (it requires an intermidiate and auxiliary TArray). This function copies the
-	 *   given InArray into the desired input FNeuralTensor.
+	 * SetInputFromArrayCopy(), SetInputFromVoidPointerCopy(), and GetInputDataPointerMutable() are the only functions that allow modifying the
+	 * network input tensor(s) values:
+	 * - SetInputFromArrayCopy() and SetInputFromVoidPointerCopy() are very easy to use but less efficient (they require an intermediate and
+	 *   auxiliary TArray or pointer of data). These functions copy the given InArray or data pointer into the desired input FNeuralTensor.
 	 * - GetInputDataPointerMutable() is potentially more efficient because it avoids creating and copying from an intermediate TArray. This function
 	 *   returns a mutable void* pointer of the desired input FNeuralTensor that can be filled by the user on-the-fly.
 	 */
 	void SetInputFromArrayCopy(const TArray<float>& InArray, const int32 InTensorIndex = 0);
-// CONTINUE HERE
 	void SetInputFromVoidPointerCopy(const void* const InVoidPtr, const int32 InTensorIndex = 0);
 	void* GetInputDataPointerMutable(const int32 InTensorIndex = 0);
 
-// @todo: Move these 3 functions to NNIQA.
-
 	/**
-	 * Non-efficient functions meant to be used only for debugging purposes.
+	 * Non-computationally-efficient functions meant to be used only for debugging purposes, but should never be used on highly performant systems:
 	 * - InputTensorsToCPU copies the CPU memory of the desired input tensor(s) to GPU (to debug InputDeviceType == ENeuralDeviceType::GPU).
 	 * - OutputTensorsToCPU copies the GPU memory of the desired output tensor(s) back to CPU (to debug OutputDeviceType == ENeuralDeviceType::GPU).
 	 * @param InTensorIndexes If empty (default value), it will apply to all output tensors.
@@ -212,16 +217,22 @@ public:
 	void OutputTensorsToCPU(const TArray<int32>& InTensorIndexes = TArray<int32>());
 
 	/**
-	 * Run() executes the forward pass on the current UNeuralNetwork given the current input FDeprecatedNeuralTensor(s), which were previously filled
-	 * with SetInputFromArrayCopy() or GetInputDataPointerMutable().
-	 * Its output results can be retrieved with GetOutputTensor() or GetOutputTensors().
+	 * Load() + SetInputFromArrayCopy() + Run() is the simplest way to load an ONNX file, set the input tensor(s), and run inference on it. All other
+	 * functions in UNeuralNetwork provide additional functionality (set the device to CPU/GPU, more complicated but faster ways to set the input,
+	 * support for asynchronous run, avoid copying the memory CPU-to-GPU or GPU-to-CPU, etc.).
 	 *
-	 * If Run() is called asynchronously, this does not guarantee that calling SetInputFromArrayCopy multiple times will result in each one being
-	 * applied for a different Run. The user is responsible of not calling SetInputFromArrayCopy until Run() is completed and its delegate
-	 * (OnAsyncRunCompletedDelegate) called. Otherwise, the wrong results might be returned.
+	 * Run() executes the forward pass of the current UNeuralNetwork given the current input FNeuralTensor(s), which were previously filled
+	 * with SetInputFromArrayCopy(), SetInputFromVoidPointerCopy(), or GetInputDataPointerMutable().
+	 * Its output results can be retrieved with GetOutputTensor().
+	 *
+	 * If Run() is called asynchronously, the user is responsible of not calling SetInputFromArrayCopy/SetInputFromVoidPointerCopy until Run() is
+	 * completed and its delegate (OnAsyncRunCompletedDelegate) called. Otherwise, the wrong results might be returned.
+	 *
+	 * @see Load(), SetInputFromArrayCopy(), Run().
 	 */
 	void Run();
 
+// CONTINUE HERE
 	/**
 	 * Stats functions:
 	 * - GetLastInferenceTime will provide the last inference time measured milliseconds
@@ -229,7 +240,7 @@ public:
 	 * - GetInputMemoryTransferStats returns Input Memory Transfer statistics. (NumberSamples, Average, StdDev, Min, Max statistics measured in
 	 *   milliseconds)
 	 */
-	float GetLastInferenceTime() const;
+	float GetLastInferenceTimeMSec() const;
 	FNeuralStatsData GetInferenceStats() const;
 	FNeuralStatsData GetInputMemoryTransferStats() const;
 	void ResetStats();
