Renamed UNN function & FNeuralStats names.

#rnx

diff --git a/Engine/Plugins/Experimental/NeuralNetworkInference/Source/NeuralNetworkInference/Private/NeuralNetwork.cpp b/Engine/Plugins/Experimental/NeuralNetworkInference/Source/NeuralNetworkInference/Private/NeuralNetwork.cpp
--- a/Engine/Plugins/Experimental/NeuralNetworkInference/Source/NeuralNetworkInference/Private/NeuralNetwork.cpp
+++ b/Engine/Plugins/Experimental/NeuralNetworkInference/Source/NeuralNetworkInference/Private/NeuralNetwork.cpp
@@ -319,7 +319,7 @@ void UNeuralNetwork::SetInputFromArrayCopy(const TArray<float>& InArray, const i
 			TEXT("UNeuralNetwork::SetInputFromArrayCopy(): Unknown [BackEnd,BackEndForCurrentPlatform] = [%d,%d]."),
 			(int32)BackEnd, (int32)BackEndForCurrentPlatform);
 	}
-	InputMemoryTransferStatsModule.StoreSample(RunTimer.Toc());
+	InputTransferStatisticsEstimator.StoreSample(RunTimer.Toc());
 }
 
 void UNeuralNetwork::SetInputFromVoidPointerCopy(const void* const InVoidPtr, const int32 InTensorIndex)
@@ -351,7 +351,7 @@ void UNeuralNetwork::SetInputFromVoidPointerCopy(const void* const InVoidPtr, co
 			TEXT("UNeuralNetwork::SetInputFromVoidPointerCopy(): Unknown [BackEnd,BackEndForCurrentPlatform] = [%d,%d]."),
 			(int32)BackEnd, (int32)BackEndForCurrentPlatform);
 	}
-	InputMemoryTransferStatsModule.StoreSample(RunTimer.Toc());
+	InputTransferStatisticsEstimator.StoreSample(RunTimer.Toc());
 }
 
 void* UNeuralNetwork::GetInputDataPointerMutable(const int32 InTensorIndex)
@@ -503,28 +503,28 @@ void UNeuralNetwork::Run()
 		UE_LOG(LogNeuralNetworkInference, Warning, TEXT("UNeuralNetwork::Run(): Unknown [BackEnd,BackEndForCurrentPlatform] = [%d,%d]."),
 			(int32)BackEnd, (int32)BackEndForCurrentPlatform);
 	}
-	ComputeStatsModule.StoreSample(RunTimer.Toc());
+	ComputeStatisticsEstimator.StoreSample(RunTimer.Toc());
 }
 
-float UNeuralNetwork::GetLastInferenceTimeMSec() const
+float UNeuralNetwork::GetLastRunTimeMSec() const
 {
-	return ComputeStatsModule.GetLastSample();
+	return ComputeStatisticsEstimator.GetLastSample();
 }
 
-FNeuralStatsData UNeuralNetwork::GetInferenceStats() const
+FNeuralStatistics UNeuralNetwork::GetInferenceStats() const
 {
-	return ComputeStatsModule.GetStats();
+	return ComputeStatisticsEstimator.GetStats();
 }
 
-FNeuralStatsData UNeuralNetwork::GetInputMemoryTransferStats() const
+FNeuralStatistics UNeuralNetwork::GetInputMemoryTransferStats() const
 {
-	return InputMemoryTransferStatsModule.GetStats();
+	return InputTransferStatisticsEstimator.GetStats();
 }
 
 void UNeuralNetwork::ResetStats()
 {
-	ComputeStatsModule.ResetStats();
-	InputMemoryTransferStatsModule.ResetStats();
+	ComputeStatisticsEstimator.ResetStats();
+	InputTransferStatisticsEstimator.ResetStats();
 }
 
 
diff --git a/Engine/Plugins/Experimental/NeuralNetworkInference/Source/NeuralNetworkInference/Public/NeuralNetwork.h b/Engine/Plugins/Experimental/NeuralNetworkInference/Source/NeuralNetworkInference/Public/NeuralNetwork.h
--- a/Engine/Plugins/Experimental/NeuralNetworkInference/Source/NeuralNetworkInference/Public/NeuralNetwork.h
+++ b/Engine/Plugins/Experimental/NeuralNetworkInference/Source/NeuralNetworkInference/Public/NeuralNetwork.h
@@ -4,7 +4,7 @@
 
 #include "CoreMinimal.h"
 #include "NeuralEnumClasses.h"
-#include "NeuralStats.h"
+#include "NeuralStatistics.h"
 #include "NeuralTensor.h"
 #include "NeuralNetwork.generated.h"
 
@@ -235,14 +235,14 @@ public:
 // CONTINUE HERE
 	/**
 	 * Stats functions:
-	 * - GetLastInferenceTime will provide the last inference time measured milliseconds
+	 * - GetLastRunTimeMSec will provide the last inference time measured milliseconds.
 	 * - GetInferenceStats returns inference time statistics. (NumberSamples, Average, StdDev, Min, Max statistics measured in milliseconds)
 	 * - GetInputMemoryTransferStats returns Input Memory Transfer statistics. (NumberSamples, Average, StdDev, Min, Max statistics measured in
 	 *   milliseconds)
 	 */
-	float GetLastInferenceTimeMSec() const;
-	FNeuralStatsData GetInferenceStats() const;
-	FNeuralStatsData GetInputMemoryTransferStats() const;
+	float GetLastRunTimeMSec() const;
+	FNeuralStatistics GetInferenceStats() const;
+	FNeuralStatistics GetInputMemoryTransferStats() const;
 	void ResetStats();
 
 protected:
@@ -311,8 +311,8 @@ private:
 	/**
 	 * Stats-related members.
 	 */
-	FNeuralStats ComputeStatsModule;
-	FNeuralStats InputMemoryTransferStatsModule;
+	FNeuralStatisticsEstimator ComputeStatisticsEstimator;
+	FNeuralStatisticsEstimator InputTransferStatisticsEstimator;
 
 	/**
 	 * Struct pointer containing the UE-and-ORT-based back end implementation.
diff --git a/Engine/Plugins/Experimental/NeuralNetworkInference/Source/Private/NeuralNetworkInferenceProfiling/Private/NeuralStatistics.cpp b/Engine/Plugins/Experimental/NeuralNetworkInference/Source/Private/NeuralNetworkInferenceProfiling/Private/NeuralStatistics.cpp
new file mode 100644
--- /dev/null
+++ b/Engine/Plugins/Experimental/NeuralNetworkInference/Source/Private/NeuralNetworkInferenceProfiling/Private/NeuralStatistics.cpp
@@ -0,0 +1,124 @@
+// Copyright Epic Games, Inc. All Rights Reserved.
+
+#include "NeuralStatistics.h"
+#include "GenericPlatform/GenericPlatformMath.h"
+
+DEFINE_LOG_CATEGORY_STATIC(LogNNIStats, Display, All);
+
+
+
+/* FNeuralStatisticsEstimator auxiliary functions
+ *****************************************************************************/
+
+void FNeuralStatisticsEstimator_CalculateMinMax(float& OutMin, float& OutMax, const TCircularBuffer<float>& InDataBuffer, const bool bInIsBufferFull,
+	const uint32 InBufferIdx)
+{
+	const uint32 PtrEnd = (bInIsBufferFull) ? InDataBuffer.Capacity() : InBufferIdx;
+
+	OutMin = FLT_MAX;
+	OutMax = FLT_MIN;
+	for (uint32 Index = 0; Index < PtrEnd; ++Index)
+	{
+		const float CurrentElement = InDataBuffer[Index];
+		if (CurrentElement < OutMin)
+		{
+			OutMin = CurrentElement;
+		}
+		if (CurrentElement > OutMax)
+		{
+			OutMax = CurrentElement;
+		}
+	}
+}
+
+
+
+/* FNeuralStatistics constructor
+ *****************************************************************************/
+
+FNeuralStatistics::FNeuralStatistics(const int InNumberSamples, const float InAverage, const float InStdDev, const float InMin, const float InMax)
+	: NumberSamples(InNumberSamples)
+	, Average(InAverage)
+	, StdDev(InStdDev)
+	, Min(InMin)
+	, Max(InMax)
+{}
+
+
+
+/* FNeuralNetworkInferenceQATimer public functions
+ *****************************************************************************/
+
+FNeuralStatisticsEstimator::FNeuralStatisticsEstimator(const uint32 InSizeRollingWindow) 
+	: EpsilonFloat(0.000001f)
+	, BufferIdx(0)
+	, bIsBufferFull(false)
+	, LastSample(0.f)
+	, DataBuffer(InSizeRollingWindow, 0)
+{
+}
+	
+void FNeuralStatisticsEstimator::StoreSample(const float InRunTime) 
+{
+	LastSample = InRunTime;
+
+	DataBuffer[BufferIdx] = LastSample;
+
+	if(BufferIdx == DataBuffer.Capacity() && !bIsBufferFull)
+	{
+		bIsBufferFull = true;
+	}
+	BufferIdx = DataBuffer.GetNextIndex(BufferIdx);
+}
+
+void FNeuralStatisticsEstimator::ResetStats()
+{
+	for (uint32 Index = 0; Index < DataBuffer.Capacity(); ++Index)
+	{
+		DataBuffer[Index] = 0;
+	}
+	bIsBufferFull = false;
+	BufferIdx = 0;
+	LastSample = 0;
+}
+
+float FNeuralStatisticsEstimator::GetLastSample() const
+{
+	return LastSample;
+}
+
+FNeuralStatistics FNeuralStatisticsEstimator::GetStats() const
+{
+	const uint32 count = (bIsBufferFull) ? DataBuffer.Capacity() : BufferIdx;
+	const float Mean = CalculateMean();
+	const float StdDev = CalculateStdDev(Mean);
+	float Min, Max;
+	FNeuralStatisticsEstimator_CalculateMinMax(Min, Max, DataBuffer, bIsBufferFull, BufferIdx);
+	// Return aggregated data
+	return FNeuralStatistics(count, Mean, StdDev, Min, Max);
+}
+
+float FNeuralStatisticsEstimator::CalculateMean() const
+{
+	uint32 PtrEnd = (bIsBufferFull) ? DataBuffer.Capacity(): BufferIdx;
+	float AccumulatorMean = 0;
+	for (uint32 Index = 0; Index < PtrEnd; ++Index)
+	{
+		AccumulatorMean += DataBuffer[Index];
+	}
+	// Return mean
+	return AccumulatorMean / (static_cast<float>(PtrEnd) + EpsilonFloat);
+}
+
+float FNeuralStatisticsEstimator::CalculateStdDev(const float InMean) const
+{
+	const uint32 PtrEnd = (bIsBufferFull) ? DataBuffer.Capacity() : BufferIdx;
+
+	float AccumulatorStdDev = 0;
+	for (uint32 Index = 0; Index < PtrEnd; ++Index)
+	{
+		AccumulatorStdDev += (DataBuffer[Index] - InMean) * (DataBuffer[Index] - InMean);
+	}
+	// Return std deviation
+	return FMath::Sqrt(AccumulatorStdDev / (static_cast<float>(PtrEnd) + EpsilonFloat));
+}
diff --git a/Engine/Plugins/Experimental/NeuralNetworkInference/Source/Private/NeuralNetworkInferenceProfiling/Private/NeuralStats.cpp b/Engine/Plugins/Experimental/NeuralNetworkInference/Source/Private/NeuralNetworkInferenceProfiling/Private/NeuralStats.cpp
deleted file mode 100644
diff --git a/Engine/Plugins/Experimental/NeuralNetworkInference/Source/Private/NeuralNetworkInferenceProfiling/Public/NeuralStatistics.h b/Engine/Plugins/Experimental/NeuralNetworkInference/Source/Private/NeuralNetworkInferenceProfiling/Public/NeuralStatistics.h
new file mode 100644
--- /dev/null
+++ b/Engine/Plugins/Experimental/NeuralNetworkInference/Source/Private/NeuralNetworkInferenceProfiling/Public/NeuralStatistics.h
@@ -0,0 +1,39 @@
+// Copyright Epic Games, Inc. All Rights Reserved.
+
+#pragma once
+
+#include "CoreMinimal.h"
+#include "Containers/CircularBuffer.h"
+
+struct NEURALNETWORKINFERENCEPROFILING_API FNeuralStatistics
+{
+	const uint32 NumberSamples;
+	const float Average;
+	const float StdDev;
+	const float Min;
+	const float Max;
+
+	FNeuralStatistics(const int InNumberSamples, const float InAverage, const float InStdDev, const float InMin, const float InMax);
+};
+
+class NEURALNETWORKINFERENCEPROFILING_API FNeuralStatisticsEstimator
+{
+public:
+	FNeuralStatisticsEstimator(const uint32 InSizeRollingWindow = 1024);
+	
+	void StoreSample(const float InRunTime);
+	void ResetStats();
+
+	float GetLastSample() const;
+	FNeuralStatistics GetStats() const;
+
+private:
+	const float EpsilonFloat;
+	uint32 BufferIdx;
+	bool bIsBufferFull;
+	float LastSample;
+	TCircularBuffer<float> DataBuffer;
+
+	float CalculateMean() const;
+	float CalculateStdDev(const float InMean) const;
+};
diff --git a/Engine/Plugins/Experimental/NeuralNetworkInference/Source/Private/NeuralNetworkInferenceProfiling/Public/NeuralStats.h b/Engine/Plugins/Experimental/NeuralNetworkInference/Source/Private/NeuralNetworkInferenceProfiling/Public/NeuralStats.h
deleted file mode 100644
diff --git a/Engine/Plugins/Experimental/NeuralNetworkInference/Source/Private/NeuralNetworkInferenceQA/Private/ModelUnitTester.cpp b/Engine/Plugins/Experimental/NeuralNetworkInference/Source/Private/NeuralNetworkInferenceQA/Private/ModelUnitTester.cpp
--- a/Engine/Plugins/Experimental/NeuralNetworkInference/Source/Private/NeuralNetworkInferenceQA/Private/ModelUnitTester.cpp
+++ b/Engine/Plugins/Experimental/NeuralNetworkInference/Source/Private/NeuralNetworkInferenceQA/Private/ModelUnitTester.cpp
@@ -12,21 +12,23 @@
 
 struct FNNIUnitTesterTimeData
 {
-	FNeuralStatsData ComputeTimeData;
-	FNeuralStatsData InputCopyTimeData;
-	FNeuralStatsData OutputCopyTimeData;
+	FNeuralStatistics ComputeTimeData;
+	FNeuralStatistics InputCopyTimeData;
+	FNeuralStatistics OutputCopyTimeData;
 
-	FNNIUnitTesterTimeData(const FNeuralStatsData& InComputeTimeData, const FNeuralStatsData& InInputCopyTimeData, const FNeuralStatsData& InOutputCopyTimeData)
+	FNNIUnitTesterTimeData(const FNeuralStatistics& InComputeTimeData, const FNeuralStatistics& InInputCopyTimeData,
+		const FNeuralStatistics& InOutputCopyTimeData)
 		: ComputeTimeData(InComputeTimeData)
 		, InputCopyTimeData(InInputCopyTimeData)
 		, OutputCopyTimeData(InOutputCopyTimeData)
 	{}
 };
 
-FNNIUnitTesterTimeData FModelUnitTester_GetTimeInformation(UNeuralNetwork* InOutNetwork, TArray<float>& OutCPUGPUCPUOutput, const TArray<float>& InInputArray, const int32 InRepetitions)
+FNNIUnitTesterTimeData FModelUnitTester_GetTimeInformation(UNeuralNetwork* InOutNetwork, TArray<float>& OutCPUGPUCPUOutput,
+	const TArray<float>& InInputArray, const int32 InRepetitions)
 {
 	/* Input/output copy speed */
-	FNeuralStats OutCopyingStats;
+	FNeuralStatisticsEstimator OutCopyingStats;
 	FNeuralTimer Timer;
 	for (int32 TimerIndex = 0; TimerIndex < InRepetitions; ++TimerIndex)
 	{
@@ -123,8 +125,11 @@ bool FModelUnitTester::GlobalTest(const FString& InProjectContentDir, const FStr
 	const TArray<int32> GPURepetitionsForUEOnlyBackEnd({ 0, 0, 0, 0 });
 #endif //WITH_EDITOR
 	// Run tests
-	UE_LOG(LogNeuralNetworkInferenceQA, Display, TEXT("ENeuralBackEnd::UEOnly test for FModelUnitTester::GlobalTest disabled due to not being compatible with ORT (linking error issues)."));
-	return ModelLoadAccuracyAndSpeedTests(InProjectContentDir, InModelZooRelativeDirectory, ModelNames, InputArrayValues, CPUGroundTruths, GPUGroundTruths, CPURepetitionsForUEAndORTBackEnd, GPURepetitionsForUEAndORTBackEnd, CPURepetitionsForUEOnlyBackEnd, GPURepetitionsForUEOnlyBackEnd);
+	UE_LOG(LogNeuralNetworkInferenceQA, Display,
+		TEXT("ENeuralBackEnd::UEOnly test for FModelUnitTester::GlobalTest disabled due to not being compatible with ORT (linking error issues)."));
+	return ModelLoadAccuracyAndSpeedTests(InProjectContentDir, InModelZooRelativeDirectory, ModelNames, InputArrayValues, CPUGroundTruths,
+		GPUGroundTruths, CPURepetitionsForUEAndORTBackEnd, GPURepetitionsForUEAndORTBackEnd, CPURepetitionsForUEOnlyBackEnd,
+		GPURepetitionsForUEOnlyBackEnd);
 }
 
 
@@ -132,9 +137,11 @@ bool FModelUnitTester::GlobalTest(const FString& InProjectContentDir, const FStr
 /* FModelUnitTester static private functions
  *****************************************************************************/
 
-bool FModelUnitTester::ModelLoadAccuracyAndSpeedTests(const FString& InProjectContentDir, const FString& InModelZooRelativeDirectory, const TArray<FString>& InModelNames,
-	const TArray<float>& InInputArrayValues, const TArray<TArray<double>>& InCPUGroundTruths, const TArray<TArray<double>>& InGPUGroundTruths, const TArray<int32>& InCPURepetitionsForUEAndORTBackEnd,
-	const TArray<int32>& InGPURepetitionsForUEAndORTBackEnd, const TArray<int32>& InCPURepetitionsForUEOnlyBackEnd, const TArray<int32>& InGPURepetitionsForUEOnlyBackEnd)
+bool FModelUnitTester::ModelLoadAccuracyAndSpeedTests(const FString& InProjectContentDir, const FString& InModelZooRelativeDirectory,
+	const TArray<FString>& InModelNames, const TArray<float>& InInputArrayValues, const TArray<TArray<double>>& InCPUGroundTruths,
+	const TArray<TArray<double>>& InGPUGroundTruths, const TArray<int32>& InCPURepetitionsForUEAndORTBackEnd,
+	const TArray<int32>& InGPURepetitionsForUEAndORTBackEnd, const TArray<int32>& InCPURepetitionsForUEOnlyBackEnd,
+	const TArray<int32>& InGPURepetitionsForUEOnlyBackEnd)
 {
 	bool bDidGlobalTestPassed = true;
 
@@ -152,30 +159,36 @@ bool FModelUnitTester::ModelLoadAccuracyAndSpeedTests(const FString& InProjectCo
 		UNeuralNetwork* Network = NetworkUassetLoadTest(UAssetModelFilePath);
 		if (!Network)
 		{
-			UE_LOG(LogNeuralNetworkInferenceQA, Warning, TEXT("UNeuralNetwork could not be loaded from UAssetModelFilePath %s."), *UAssetModelFilePath);
+			UE_LOG(LogNeuralNetworkInferenceQA, Warning, TEXT("UNeuralNetwork could not be loaded from UAssetModelFilePath %s."),
+				*UAssetModelFilePath);
 			return false;
 		}
 		// Input debugging
 		UE_LOG(LogNeuralNetworkInferenceQA, Display, TEXT("--------------- %s - Input/Output"), *ModelName);
 		for (int32 TensorIndex = 0; TensorIndex < Network->GetInputTensorNumber(); ++TensorIndex)
 		{
-			UE_LOG(LogNeuralNetworkInferenceQA, Display, TEXT("InputTensor[%d] = %s."), TensorIndex, *Network->GetInputTensor(TensorIndex).GetName());
+			UE_LOG(LogNeuralNetworkInferenceQA, Display, TEXT("InputTensor[%d] = %s."), TensorIndex,
+				*Network->GetInputTensor(TensorIndex).GetName());
 		}
 		// Output debugging
 		for (int32 TensorIndex = 0; TensorIndex < Network->GetOutputTensorNumber(); ++TensorIndex)
 		{
-			UE_LOG(LogNeuralNetworkInferenceQA, Display, TEXT("OutputTensor[%d] = %s."), TensorIndex, *Network->GetOutputTensor(TensorIndex).GetName());
+			UE_LOG(LogNeuralNetworkInferenceQA, Display, TEXT("OutputTensor[%d] = %s."), TensorIndex,
+				*Network->GetOutputTensor(TensorIndex).GetName());
 		}
 		const bool bShouldRunUEAndORTBackEnd = (InCPURepetitionsForUEAndORTBackEnd[ModelIndex] + InGPURepetitionsForUEAndORTBackEnd[ModelIndex] > 0);
 		if (bShouldRunUEAndORTBackEnd)
 		{
-			bDidGlobalTestPassed &= ModelAccuracyTest(Network, ENeuralSynchronousMode::Synchronous, UNeuralNetwork::ENeuralBackEnd::UEAndORT, InInputArrayValues, CPUGroundTruths, GPUGroundTruths);
-			bDidGlobalTestPassed &= ModelAccuracyTest(Network, ENeuralSynchronousMode::Asynchronous, UNeuralNetwork::ENeuralBackEnd::UEAndORT, InInputArrayValues, CPUGroundTruths, GPUGroundTruths);
+			bDidGlobalTestPassed &= ModelAccuracyTest(Network, ENeuralSynchronousMode::Synchronous, UNeuralNetwork::ENeuralBackEnd::UEAndORT,
+				InInputArrayValues, CPUGroundTruths, GPUGroundTruths);
+			bDidGlobalTestPassed &= ModelAccuracyTest(Network, ENeuralSynchronousMode::Asynchronous, UNeuralNetwork::ENeuralBackEnd::UEAndORT,
+				InInputArrayValues, CPUGroundTruths, GPUGroundTruths);
 		}
 		const bool bShouldRunUEOnlyBackEnd = false; // (InCPURepetitionsForUEOnlyBackEnd[ModelIndex] * InCPURepetitionsForUEOnlyBackEnd[ModelIndex] > 0);
 		if (bShouldRunUEOnlyBackEnd)
 		{
-			bDidGlobalTestPassed &= ModelAccuracyTest(Network, ENeuralSynchronousMode::Synchronous, UNeuralNetwork::ENeuralBackEnd::UEOnly, InInputArrayValues, CPUGroundTruths, GPUGroundTruths);
+			bDidGlobalTestPassed &= ModelAccuracyTest(Network, ENeuralSynchronousMode::Synchronous, UNeuralNetwork::ENeuralBackEnd::UEOnly,
+				InInputArrayValues, CPUGroundTruths, GPUGroundTruths);
 		}
 
 		UE_LOG(LogNeuralNetworkInferenceQA, Display, TEXT("---------------------------------------------------------------------------------------------------------------------------------"));
@@ -184,17 +197,22 @@ bool FModelUnitTester::ModelLoadAccuracyAndSpeedTests(const FString& InProjectCo
 #if WITH_EDITOR
 		const int32 ONNXOrORTIndex = 0; //for (int32 ONNXOrORTIndex = 0; ONNXOrORTIndex < 2; ++ONNXOrORTIndex)
 		{
-			const FString ModelFilePath = (ONNXOrORTIndex % 2 == 0 ? GetONNXModelFilePath(ModelZooDirectory, ModelName) : GetORTModelFilePath(ModelZooDirectory, ModelName));
+			const FString ModelFilePath = (ONNXOrORTIndex % 2 == 0 ? GetONNXModelFilePath(ModelZooDirectory, ModelName)
+				: GetORTModelFilePath(ModelZooDirectory, ModelName));
 			const FString ModelType = (ONNXOrORTIndex % 2 == 0 ? TEXT("ONNX") : TEXT("ORT"));
-			UE_LOG(LogNeuralNetworkInferenceQA, Display, TEXT("-------------------- %s - Network %s Load and Run - %s"), *ModelName, *ModelType, *ModelFilePath);
+			UE_LOG(LogNeuralNetworkInferenceQA, Display, TEXT("-------------------- %s - Network %s Load and Run - %s"), *ModelName, *ModelType,
+				*ModelFilePath);
 			if (bShouldRunUEAndORTBackEnd)
 			{
-				bDidGlobalTestPassed &= ModelAccuracyTest(NetworkONNXOrORTLoadTest(ModelFilePath), ENeuralSynchronousMode::Synchronous, UNeuralNetwork::ENeuralBackEnd::UEAndORT, InInputArrayValues, CPUGroundTruths, GPUGroundTruths);
-				bDidGlobalTestPassed &= ModelAccuracyTest(NetworkONNXOrORTLoadTest(ModelFilePath), ENeuralSynchronousMode::Asynchronous, UNeuralNetwork::ENeuralBackEnd::UEAndORT, InInputArrayValues, CPUGroundTruths, GPUGroundTruths);
+				bDidGlobalTestPassed &= ModelAccuracyTest(NetworkONNXOrORTLoadTest(ModelFilePath), ENeuralSynchronousMode::Synchronous,
+					UNeuralNetwork::ENeuralBackEnd::UEAndORT, InInputArrayValues, CPUGroundTruths, GPUGroundTruths);
+				bDidGlobalTestPassed &= ModelAccuracyTest(NetworkONNXOrORTLoadTest(ModelFilePath), ENeuralSynchronousMode::Asynchronous,
+					UNeuralNetwork::ENeuralBackEnd::UEAndORT, InInputArrayValues, CPUGroundTruths, GPUGroundTruths);
 			}
 			if (bShouldRunUEOnlyBackEnd)
 			{
-				bDidGlobalTestPassed &= ModelAccuracyTest(NetworkONNXOrORTLoadTest(ModelFilePath), ENeuralSynchronousMode::Synchronous, UNeuralNetwork::ENeuralBackEnd::UEOnly, InInputArrayValues, CPUGroundTruths, GPUGroundTruths);
+				bDidGlobalTestPassed &= ModelAccuracyTest(NetworkONNXOrORTLoadTest(ModelFilePath), ENeuralSynchronousMode::Synchronous,
+					UNeuralNetwork::ENeuralBackEnd::UEOnly, InInputArrayValues, CPUGroundTruths, GPUGroundTruths);
 			}
 		}
 #else //WITH_EDITOR
@@ -211,7 +229,8 @@ bool FModelUnitTester::ModelLoadAccuracyAndSpeedTests(const FString& InProjectCo
 
 		// UEAndORT (if WITH_UE_AND_ORT_SUPPORT)
 #ifdef WITH_UE_AND_ORT_SUPPORT
-		bDidGlobalTestPassed &= ModelSpeedTest(UAssetModelFilePath, ENeuralDeviceType::CPU, UNeuralNetwork::ENeuralBackEnd::UEAndORT, InCPURepetitionsForUEAndORTBackEnd[ModelIndex]);
+		bDidGlobalTestPassed &= ModelSpeedTest(UAssetModelFilePath, ENeuralDeviceType::CPU, UNeuralNetwork::ENeuralBackEnd::UEAndORT,
+			InCPURepetitionsForUEAndORTBackEnd[ModelIndex]);
 		UNeuralNetwork* Network = NetworkUassetLoadTest(UAssetModelFilePath);
 		if (Network->GetBackEndForCurrentPlatform() != UNeuralNetwork::ENeuralBackEnd::UEAndORT)
 		{
@@ -220,12 +239,13 @@ bool FModelUnitTester::ModelLoadAccuracyAndSpeedTests(const FString& InProjectCo
 		}
 		else if (InGPURepetitionsForUEAndORTBackEnd[ModelIndex] > 0 && Network->IsGPUSupported())
 		{
-			bDidGlobalTestPassed &= ModelSpeedTest(UAssetModelFilePath, ENeuralDeviceType::GPU, UNeuralNetwork::ENeuralBackEnd::UEAndORT, InGPURepetitionsForUEAndORTBackEnd[ModelIndex]);
+			bDidGlobalTestPassed &= ModelSpeedTest(UAssetModelFilePath, ENeuralDeviceType::GPU, UNeuralNetwork::ENeuralBackEnd::UEAndORT,
+				InGPURepetitionsForUEAndORTBackEnd[ModelIndex]);
 		}
 		else
 		{
-			UE_LOG(LogNeuralNetworkInferenceQA, Display, TEXT("-------------------- ModelSpeedTest-UEAndORT-GPU skipped (DX12 not enabled or InGPURepetitionsForUEAndORTBackEnd[ModelIndex] = %d is 0)."),
-				InGPURepetitionsForUEAndORTBackEnd[ModelIndex]);
+			UE_LOG(LogNeuralNetworkInferenceQA, Display, TEXT("-------------------- ModelSpeedTest-UEAndORT-GPU skipped (DX12 not enabled or"
+				" InGPURepetitionsForUEAndORTBackEnd[ModelIndex] = %d is 0)."), InGPURepetitionsForUEAndORTBackEnd[ModelIndex]);
 		}
 #endif //WITH_UE_AND_ORT_SUPPORT
 
@@ -233,8 +253,10 @@ bool FModelUnitTester::ModelLoadAccuracyAndSpeedTests(const FString& InProjectCo
 		const bool bShouldRunUEOnlyBackEnd = false;
 		if (bShouldRunUEOnlyBackEnd)
 		{
-			bDidGlobalTestPassed &= ModelSpeedTest(UAssetModelFilePath, ENeuralDeviceType::CPU, UNeuralNetwork::ENeuralBackEnd::UEOnly, InCPURepetitionsForUEOnlyBackEnd[ModelIndex]);
-			bDidGlobalTestPassed &= ModelSpeedTest(UAssetModelFilePath, ENeuralDeviceType::GPU, UNeuralNetwork::ENeuralBackEnd::UEOnly, InGPURepetitionsForUEOnlyBackEnd[ModelIndex]);
+			bDidGlobalTestPassed &= ModelSpeedTest(UAssetModelFilePath, ENeuralDeviceType::CPU, UNeuralNetwork::ENeuralBackEnd::UEOnly,
+				InCPURepetitionsForUEOnlyBackEnd[ModelIndex]);
+			bDidGlobalTestPassed &= ModelSpeedTest(UAssetModelFilePath, ENeuralDeviceType::GPU, UNeuralNetwork::ENeuralBackEnd::UEOnly,
+				InGPURepetitionsForUEOnlyBackEnd[ModelIndex]);
 		}
 	}
 
@@ -243,17 +265,20 @@ bool FModelUnitTester::ModelLoadAccuracyAndSpeedTests(const FString& InProjectCo
 
 FString FModelUnitTester::GetONNXModelFilePath(const FString& ModelZooDirectory, const FString& InModelName)
 {
-	return /*FPaths::ConvertRelativePathToFull*/(ModelZooDirectory / InModelName + TEXT("/") + InModelName + TEXT(".onnx")); // E.g., ModelZooDirectory / TEXT("ExampleNetworkReadable/ExampleNetworkReadable.onnx")
+	// E.g., ModelZooDirectory / TEXT("ExampleNetworkReadable/ExampleNetworkReadable.onnx")
+	return /*FPaths::ConvertRelativePathToFull*/(ModelZooDirectory / InModelName + TEXT("/") + InModelName + TEXT(".onnx"));
 }
 
 FString FModelUnitTester::GetORTModelFilePath(const FString& ModelZooDirectory, const FString& InModelName)
 {
-	return /*FPaths::ConvertRelativePathToFull*/(ModelZooDirectory / InModelName + TEXT("/") + InModelName + TEXT(".ort")); // E.g., ModelZooDirectory / TEXT("ExampleNetworkReadable/ExampleNetworkReadable.ort")
+	// E.g., ModelZooDirectory / TEXT("ExampleNetworkReadable/ExampleNetworkReadable.ort")
+	return /*FPaths::ConvertRelativePathToFull*/(ModelZooDirectory / InModelName + TEXT("/") + InModelName + TEXT(".ort"));
 }
 
 FString FModelUnitTester::GetUAssetModelFilePath(const FString& InModelName, const FString& InModelZooRelativeDirectory)
 {
-	return InModelName + TEXT("'/Game/") + InModelZooRelativeDirectory / InModelName + TEXT("/") + InModelName + TEXT(".") + InModelName + TEXT("'"); // E.g., '/Game/[MODEL_ZOO_DIR]/ExampleNetworkReadable/ExampleNetworkReadable.ExampleNetworkReadable'
+	// E.g., '/Game/[MODEL_ZOO_DIR]/ExampleNetworkReadable/ExampleNetworkReadable.ExampleNetworkReadable'
+	return InModelName + TEXT("'/Game/") + InModelZooRelativeDirectory / InModelName + TEXT("/") + InModelName + TEXT(".") + InModelName + TEXT("'");
 }
 
 UNeuralNetwork* FModelUnitTester::NetworkUassetLoadTest(const FString& InUAssetPath)
@@ -289,8 +314,9 @@ UNeuralNetwork* FModelUnitTester::NetworkONNXOrORTLoadTest(const FString& InMode
 	return Network;
 }
 
-bool FModelUnitTester::ModelAccuracyTest(UNeuralNetwork* InOutNetwork, const ENeuralSynchronousMode InSynchronousMode, const UNeuralNetwork::ENeuralBackEnd InBackEnd, const TArray<float>& InInputArrayValues,
-	const TArray<double>& InCPUGroundTruths, const TArray<double>& InGPUGroundTruths)
+bool FModelUnitTester::ModelAccuracyTest(UNeuralNetwork* InOutNetwork, const ENeuralSynchronousMode InSynchronousMode,
+	const UNeuralNetwork::ENeuralBackEnd InBackEnd, const TArray<float>& InInputArrayValues, const TArray<double>& InCPUGroundTruths,
+	const TArray<double>& InGPUGroundTruths)
 {
 	// Sanity check
 	if (!InOutNetwork)
@@ -310,7 +336,8 @@ bool FModelUnitTester::ModelAccuracyTest(UNeuralNetwork* InOutNetwork, const ENe
 		}
 		if (InputArrays.Num() > InCPUGroundTruths.Num() || InputArrays.Num() > InGPUGroundTruths.Num())
 		{
-			UE_LOG(LogNeuralNetworkInferenceQA, Warning, TEXT("InputArrays.Num() <= InCPUGroundTruths.Num() && InputArrays.Num() <= InGPUGroundTruths.Num() failed: %d vs. %d vs. %d."),
+			UE_LOG(LogNeuralNetworkInferenceQA, Warning,
+				TEXT("InputArrays.Num() <= InCPUGroundTruths.Num() && InputArrays.Num() <= InGPUGroundTruths.Num() failed: %d vs. %d vs. %d."),
 				InputArrays.Num(), InCPUGroundTruths.Num(), InGPUGroundTruths.Num());
 			return false;
 		}
@@ -342,16 +369,20 @@ bool FModelUnitTester::ModelAccuracyTest(UNeuralNetwork* InOutNetwork, const ENe
 	TArray<TArray<float>> CPUOutputs, CPUGPUCPUOutputs, CPUGPUGPUOutputs, GPUGPUCPUOutputs;
 	
 	// Input CPU + Network CPU + Output CPU
-	ModelAccuracyTestRun(CPUOutputs, InOutNetwork, InputArrays, /*DeviceType*/ENeuralDeviceType::CPU, /*InputDeviceType*/ENeuralDeviceType::CPU, /*OutputDeviceType*/ENeuralDeviceType::CPU);
+	ModelAccuracyTestRun(CPUOutputs, InOutNetwork, InputArrays, /*DeviceType*/ENeuralDeviceType::CPU, /*InputDeviceType*/ENeuralDeviceType::CPU,
+		/*OutputDeviceType*/ENeuralDeviceType::CPU);
 
 	if (InOutNetwork->IsGPUSupported())
 	{
 		// Input CPU + Network GPU + Output CPU
-		ModelAccuracyTestRun(CPUGPUCPUOutputs, InOutNetwork, InputArrays, /*DeviceType*/ENeuralDeviceType::GPU, /*InputDeviceType*/ENeuralDeviceType::CPU, /*OutputDeviceType*/ENeuralDeviceType::CPU);
+		ModelAccuracyTestRun(CPUGPUCPUOutputs, InOutNetwork, InputArrays, /*DeviceType*/ENeuralDeviceType::GPU,
+			/*InputDeviceType*/ENeuralDeviceType::CPU, /*OutputDeviceType*/ENeuralDeviceType::CPU);
 		// Input CPU + Network GPU + Output GPU
-		ModelAccuracyTestRun(CPUGPUGPUOutputs, InOutNetwork, InputArrays, /*DeviceType*/ENeuralDeviceType::GPU, /*InputDeviceType*/ENeuralDeviceType::CPU, /*OutputDeviceType*/ENeuralDeviceType::GPU);
+		ModelAccuracyTestRun(CPUGPUGPUOutputs, InOutNetwork, InputArrays, /*DeviceType*/ENeuralDeviceType::GPU,
+			/*InputDeviceType*/ENeuralDeviceType::CPU, /*OutputDeviceType*/ENeuralDeviceType::GPU);
 		// Input GPU + Network GPU + Output CPU
-		ModelAccuracyTestRun(GPUGPUCPUOutputs, InOutNetwork, InputArrays, /*DeviceType*/ENeuralDeviceType::GPU, /*InputDeviceType*/ENeuralDeviceType::GPU, /*OutputDeviceType*/ENeuralDeviceType::CPU);
+		ModelAccuracyTestRun(GPUGPUCPUOutputs, InOutNetwork, InputArrays, /*DeviceType*/ENeuralDeviceType::GPU,
+			/*InputDeviceType*/ENeuralDeviceType::GPU, /*OutputDeviceType*/ENeuralDeviceType::CPU);
 	}
 	else
 	{
@@ -370,15 +401,21 @@ bool FModelUnitTester::ModelAccuracyTest(UNeuralNetwork* InOutNetwork, const ENe
 		const double CPUGPUGPUAvgL1Norm = GetAveragedL1Norm(CPUGPUGPUOutput);
 		const double RelativeCoefficient = 1. / FMath::Max(1., FMath::Min(CPUAvgL1Norm, CPUGPUCPUAvgL1Norm)); // Max(1, X) to avoid 0s
 		const double CPUGPUAvgL1NormDiff = GetAveragedL1NormDiff(CPUOutputs[Index], CPUGPUCPUOutput) * RelativeCoefficient * 1e6;
-		const double GPUGPUInputAvgL1NormDiff = GetAveragedL1NormDiff(CPUGPUCPUOutput, GPUGPUCPUOutput) / FMath::Max(1., FMath::Min(GPUGPUCPUAvgL1Norm, CPUGPUCPUAvgL1Norm)) * 1e12;
-		const double GPUGPUOutputAvgL1NormDiff = GetAveragedL1NormDiff(CPUGPUCPUOutput, CPUGPUGPUOutput) / FMath::Max(1., FMath::Min(CPUGPUGPUAvgL1Norm, CPUGPUCPUAvgL1Norm)) * 1e12;
+		const double GPUGPUInputAvgL1NormDiff = GetAveragedL1NormDiff(CPUGPUCPUOutput, GPUGPUCPUOutput)
+			/ FMath::Max(1., FMath::Min(GPUGPUCPUAvgL1Norm, CPUGPUCPUAvgL1Norm)) * 1e12;
+		const double GPUGPUOutputAvgL1NormDiff = GetAveragedL1NormDiff(CPUGPUCPUOutput, CPUGPUGPUOutput)
+			/ FMath::Max(1., FMath::Min(CPUGPUGPUAvgL1Norm, CPUGPUCPUAvgL1Norm)) * 1e12;
 		const double FastCPUGPUAvgL1NormDiff = FMath::Abs((CPUAvgL1Norm - CPUGPUCPUAvgL1Norm)) * RelativeCoefficient * 1e6;
-		const double FastCPUAvgL1NormDiff = FMath::Abs(CPUAvgL1Norm - InCPUGroundTruths[Index]) / FMath::Max(1., FMath::Min(CPUAvgL1Norm, InCPUGroundTruths[Index])) * 1e7;
-		const double FastGPUAvgL1NormDiff = FMath::Abs(CPUGPUCPUAvgL1Norm - InGPUGroundTruths[Index]) / FMath::Max(1., FMath::Min(CPUGPUCPUAvgL1Norm, InGPUGroundTruths[Index])) * 1e7;
+		const double FastCPUAvgL1NormDiff = FMath::Abs(CPUAvgL1Norm - InCPUGroundTruths[Index])
+			/ FMath::Max(1., FMath::Min(CPUAvgL1Norm, InCPUGroundTruths[Index])) * 1e7;
+		const double FastGPUAvgL1NormDiff = FMath::Abs(CPUGPUCPUAvgL1Norm - InGPUGroundTruths[Index])
+			/ FMath::Max(1., FMath::Min(CPUGPUCPUAvgL1Norm, InGPUGroundTruths[Index])) * 1e7;
 		// Print verbose
-		UE_LOG(LogNeuralNetworkInferenceQA, Display, TEXT("%s: InputNorm = %f, OutputNormCPU = %f, OutputNormGPU = %f, OutputNormCPUGPUGPU = %f, OutputNormGT = %f, CPUAvgL1Norm = %f, CPUGPUCPUAvgL1Norm = %f,"),
+		UE_LOG(LogNeuralNetworkInferenceQA, Display,
+			TEXT("%s: InputNorm = %f, OutputNormCPU = %f, OutputNormGPU = %f, OutputNormCPUGPUGPU = %f, OutputNormGT = %f, CPUAvgL1Norm = %f, CPUGPUCPUAvgL1Norm = %f,"),
 			*BackEndString, GetAveragedL1Norm(InputArrays[Index]), CPUAvgL1Norm, CPUGPUCPUAvgL1Norm, CPUGPUGPUAvgL1Norm, InCPUGroundTruths[Index], CPUAvgL1Norm, CPUGPUCPUAvgL1Norm);
-		UE_LOG(LogNeuralNetworkInferenceQA, Display, TEXT("\tCPUGPUAvgL1NormDiff = %fe-6, GPUGPUInputAvgL1NormDiff = %fe-7, GPUGPUOutputAvgL1NormDiff = %fe-7, FastCPUGPUAvgL1NormDiff = %fe-6, FastCPUAvgL1NormDiff = %fe-7, FastGPUAvgL1NormDiff = %fe-7 (1e-7 is roughly the precision for float)."),
+		UE_LOG(LogNeuralNetworkInferenceQA, Display,
+			TEXT("\tCPUGPUAvgL1NormDiff = %fe-6, GPUGPUInputAvgL1NormDiff = %fe-7, GPUGPUOutputAvgL1NormDiff = %fe-7, FastCPUGPUAvgL1NormDiff = %fe-6, FastCPUAvgL1NormDiff = %fe-7, FastGPUAvgL1NormDiff = %fe-7 (1e-7 is roughly the precision for float)."),
 			CPUGPUAvgL1NormDiff, GPUGPUInputAvgL1NormDiff, GPUGPUOutputAvgL1NormDiff, FastCPUGPUAvgL1NormDiff, FastCPUAvgL1NormDiff, FastGPUAvgL1NormDiff);
 		// Check if test failed and (if so) display information
 		const bool bDidSomeTestFailed = (!FMath::IsFinite(FastCPUGPUAvgL1NormDiff) || FastCPUGPUAvgL1NormDiff > 5)
@@ -394,25 +431,32 @@ bool FModelUnitTester::ModelAccuracyTest(UNeuralNetwork* InOutNetwork, const ENe
 			const TArray<int64>& OutputSizes = InOutNetwork->GetOutputTensor().GetSizes();
 			UE_LOG(LogNeuralNetworkInferenceQA, Display, TEXT("FastCPUGPUAvgL1NormDiff (%fe-6) < 5e-6 might have failed."), FastCPUGPUAvgL1NormDiff);
 			UE_LOG(LogNeuralNetworkInferenceQA, Display, TEXT("CPUGPUAvgL1NormDiff (%fe-6) < 5e-6 might have failed."), CPUGPUAvgL1NormDiff);
-			UE_LOG(LogNeuralNetworkInferenceQA, Display, TEXT("GPUGPUInputAvgL1NormDiff (%fe-12) < 1e-12 might have failed."), GPUGPUInputAvgL1NormDiff);
-			UE_LOG(LogNeuralNetworkInferenceQA, Display, TEXT("GPUGPUOutputAvgL1NormDiff (%fe-12) < 1e-12 might have failed."), GPUGPUOutputAvgL1NormDiff);
+			UE_LOG(LogNeuralNetworkInferenceQA, Display, TEXT("GPUGPUInputAvgL1NormDiff (%fe-12) < 1e-12 might have failed."),
+				GPUGPUInputAvgL1NormDiff);
+			UE_LOG(LogNeuralNetworkInferenceQA, Display, TEXT("GPUGPUOutputAvgL1NormDiff (%fe-12) < 1e-12 might have failed."),
+				GPUGPUOutputAvgL1NormDiff);
 			UE_LOG(LogNeuralNetworkInferenceQA, Display, TEXT("FastCPUAvgL1NormDiff (%fe-7) < 30e-7 might have failed (~30 times the float precision).\nCPUOutput = %s."),
 				FastCPUAvgL1NormDiff, *FNeuralTensor(CPUOutputs[Index], OutputSizes).ToString(MaxNumberElementsToDisplay));
 			UE_LOG(LogNeuralNetworkInferenceQA, Display, TEXT("FastGPUAvgL1NormDiff (%fe-7) < 30e-7 might have failed (~30 times the float precision).\nCPUGPUCPUOutput = %s."),
 				FastGPUAvgL1NormDiff, *FNeuralTensor(CPUGPUCPUOutput, OutputSizes).ToString(MaxNumberElementsToDisplay));
 			UE_LOG(LogNeuralNetworkInferenceQA, Display, TEXT("Input = %s"),
 				*FNeuralTensor(InOutNetwork->GetInputTensor().GetArrayCopy<float>(), InputSizes).ToString(MaxNumberElementsToDisplay));
-			UE_LOG(LogNeuralNetworkInferenceQA, Display, TEXT("CPUOutput = %s"), *FNeuralTensor(CPUOutputs[Index], OutputSizes).ToString(MaxNumberElementsToDisplay));
-			UE_LOG(LogNeuralNetworkInferenceQA, Display, TEXT("CPUGPUCPUOutput = %s"), *FNeuralTensor(CPUGPUCPUOutput, OutputSizes).ToString(MaxNumberElementsToDisplay));
-			UE_LOG(LogNeuralNetworkInferenceQA, Display, TEXT("GPUGPUCPUOutput = %s"), *FNeuralTensor(GPUGPUCPUOutput, OutputSizes).ToString(MaxNumberElementsToDisplay));
-			UE_LOG(LogNeuralNetworkInferenceQA, Display, TEXT("CPUGPUGPUOutput = %s"), *FNeuralTensor(CPUGPUGPUOutput, OutputSizes).ToString(MaxNumberElementsToDisplay));
+			UE_LOG(LogNeuralNetworkInferenceQA, Display, TEXT("CPUOutput = %s"),
+				*FNeuralTensor(CPUOutputs[Index], OutputSizes).ToString(MaxNumberElementsToDisplay));
+			UE_LOG(LogNeuralNetworkInferenceQA, Display, TEXT("CPUGPUCPUOutput = %s"),
+				*FNeuralTensor(CPUGPUCPUOutput, OutputSizes).ToString(MaxNumberElementsToDisplay));
+			UE_LOG(LogNeuralNetworkInferenceQA, Display, TEXT("GPUGPUCPUOutput = %s"),
+				*FNeuralTensor(GPUGPUCPUOutput, OutputSizes).ToString(MaxNumberElementsToDisplay));
+			UE_LOG(LogNeuralNetworkInferenceQA, Display, TEXT("CPUGPUGPUOutput = %s"),
+				*FNeuralTensor(CPUGPUGPUOutput, OutputSizes).ToString(MaxNumberElementsToDisplay));
 			UE_LOG(LogNeuralNetworkInferenceQA, Warning, TEXT("At least 1 of the 6 CPU/GPU tests failed."));
 			return false;
 		}
 		bDidGlobalTestPassed &= !bDidSomeTestFailed;
 	}
 	// Reset to original network state
-	InOutNetwork->SetDeviceType(/*DeviceType*/OriginalDeviceType, /*InputDeviceType*/OriginalInputDeviceType, /*OutputDeviceType*/OriginalOutputDeviceType);
+	InOutNetwork->SetDeviceType(/*DeviceType*/OriginalDeviceType, /*InputDeviceType*/OriginalInputDeviceType,
+		/*OutputDeviceType*/OriginalOutputDeviceType);
 	InOutNetwork->SetSynchronousMode(OriginalSynchronousMode);
 	InOutNetwork->SetThreadModeDelegateForAsyncRunCompleted(OriginalThreadModeDelegateForAsyncRunCompleted);
 	InOutNetwork->SetBackEnd(OriginalBackEnd);
@@ -420,8 +464,9 @@ bool FModelUnitTester::ModelAccuracyTest(UNeuralNetwork* InOutNetwork, const ENe
 	return bDidGlobalTestPassed;
 }
 
-void FModelUnitTester::ModelAccuracyTestRun(TArray<TArray<float>>& OutOutputs, UNeuralNetwork* InOutNetwork, const TArray<TArray<float>> InInputArrays,
-	const ENeuralDeviceType InDeviceType, const ENeuralDeviceType InInputDeviceType, const ENeuralDeviceType InOutputDeviceType)
+void FModelUnitTester::ModelAccuracyTestRun(TArray<TArray<float>>& OutOutputs, UNeuralNetwork* InOutNetwork,
+	const TArray<TArray<float>> InInputArrays, const ENeuralDeviceType InDeviceType, const ENeuralDeviceType InInputDeviceType,
+	const ENeuralDeviceType InOutputDeviceType)
 {
 	std::atomic<bool> bDidAsyncNetworkFinished(false);
 	const ENeuralSynchronousMode SynchronousMode = InOutNetwork->GetSynchronousMode();
@@ -462,7 +507,8 @@ void FModelUnitTester::ModelAccuracyTestRun(TArray<TArray<float>>& OutOutputs, U
 	}
 }
 
-bool FModelUnitTester::ModelSpeedTest(const FString& InUAssetPath, const ENeuralDeviceType InDeviceType, const UNeuralNetwork::ENeuralBackEnd InBackEnd, const int32 InRepetitions)
+bool FModelUnitTester::ModelSpeedTest(const FString& InUAssetPath, const ENeuralDeviceType InDeviceType,
+	const UNeuralNetwork::ENeuralBackEnd InBackEnd, const int32 InRepetitions)
 {
 	// Get debug strings
 	const FString DeviceTypeString = GetDeviceTypeString(InDeviceType);
@@ -470,7 +516,9 @@ bool FModelUnitTester::ModelSpeedTest(const FString& InUAssetPath, const ENeural
 	// Skip test if reps = 0
 	if (InRepetitions < 1)
 	{
-		UE_LOG(LogNeuralNetworkInferenceQA, Display, TEXT("ModelSpeedTest skipped for uasset \"%s\" for InDeviceType = %s (%d) and InBackEnd = %s (%d)."), *InUAssetPath, *DeviceTypeString, (int32)InDeviceType, *BackEndString, (int32)InBackEnd);
+		UE_LOG(LogNeuralNetworkInferenceQA, Display,
+			TEXT("ModelSpeedTest skipped for uasset \"%s\" for InDeviceType = %s (%d) and InBackEnd = %s (%d)."),
+			*InUAssetPath, *DeviceTypeString, (int32)InDeviceType, *BackEndString, (int32)InBackEnd);
 		return true;
 	}
 	// Load Network
@@ -478,7 +526,8 @@ bool FModelUnitTester::ModelSpeedTest(const FString& InUAssetPath, const ENeural
 	// Sanity checks
 	if (!InOutNetwork)
 	{
-		UE_LOG(LogNeuralNetworkInferenceQA, Warning, TEXT("FModelUnitTester::ModelSpeedTest(): InOutNetwork was a nullptr. Path: \"%s\"."), *InUAssetPath);
+		UE_LOG(LogNeuralNetworkInferenceQA, Warning, TEXT("FModelUnitTester::ModelSpeedTest(): InOutNetwork was a nullptr. Path: \"%s\"."),
+			*InUAssetPath);
 		return false;
 	}
 	// Save original network state
@@ -508,8 +557,9 @@ bool FModelUnitTester::ModelSpeedTest(const FString& InUAssetPath, const ENeural
 	// Display speed times
 	UE_LOG(LogNeuralNetworkInferenceQA, Display,
 		TEXT("%s-%s:, Times(msec)-> Inference = %f, Copy In = %f, Copy Out %f, avg(%d times) = Inference = %f, Copy In = %f, Copy Out %f."),
-		*BackEndString, *DeviceTypeString, TimeData1.ComputeTimeData.Average, TimeData1.InputCopyTimeData.Average, TimeData1.OutputCopyTimeData.Average,
-		InRepetitions, TimeDataN.ComputeTimeData.Average, TimeDataN.InputCopyTimeData.Average, TimeDataN.OutputCopyTimeData.Average);
+		*BackEndString, *DeviceTypeString, TimeData1.ComputeTimeData.Average, TimeData1.InputCopyTimeData.Average,
+		TimeData1.OutputCopyTimeData.Average, InRepetitions, TimeDataN.ComputeTimeData.Average, TimeDataN.InputCopyTimeData.Average,
+		TimeDataN.OutputCopyTimeData.Average);
 
 	// Reset to original network state
 	InOutNetwork->SetDeviceType(OriginalDeviceType);
