1. UNN fully cleaned / commented.
2. Deprecated ENeuralBackEnd struct.

#rnx

diff --git a/Engine/Plugins/Experimental/NeuralNetworkInference/Source/NeuralNetworkInference/Private/NeuralNetwork.cpp b/Engine/Plugins/Experimental/NeuralNetworkInference/Source/NeuralNetworkInference/Private/NeuralNetwork.cpp
--- a/Engine/Plugins/Experimental/NeuralNetworkInference/Source/NeuralNetworkInference/Private/NeuralNetwork.cpp
+++ b/Engine/Plugins/Experimental/NeuralNetworkInference/Source/NeuralNetworkInference/Private/NeuralNetwork.cpp
@@ -503,17 +503,17 @@ void UNeuralNetwork::Run()
 		UE_LOG(LogNeuralNetworkInference, Warning, TEXT("UNeuralNetwork::Run(): Unknown [BackEnd,BackEndForCurrentPlatform] = [%d,%d]."),
 			(int32)BackEnd, (int32)BackEndForCurrentPlatform);
 	}
-	ComputeStatisticsEstimator.StoreSample(RunTimer.Toc());
+	RunStatisticsEstimator.StoreSample(RunTimer.Toc());
 }
 
 float UNeuralNetwork::GetLastRunTimeMSec() const
 {
-	return ComputeStatisticsEstimator.GetLastSample();
+	return RunStatisticsEstimator.GetLastSample();
 }
 
-FNeuralStatistics UNeuralNetwork::GetInferenceStats() const
+FNeuralStatistics UNeuralNetwork::GetRunStatistics() const
 {
-	return ComputeStatisticsEstimator.GetStats();
+	return RunStatisticsEstimator.GetStats();
 }
 
 FNeuralStatistics UNeuralNetwork::GetInputMemoryTransferStats() const
@@ -523,7 +523,7 @@ FNeuralStatistics UNeuralNetwork::GetInputMemoryTransferStats() const
 
 void UNeuralNetwork::ResetStats()
 {
-	ComputeStatisticsEstimator.ResetStats();
+	RunStatisticsEstimator.ResetStats();
 	InputTransferStatisticsEstimator.ResetStats();
 }
 
diff --git a/Engine/Plugins/Experimental/NeuralNetworkInference/Source/NeuralNetworkInference/Public/NeuralNetwork.h b/Engine/Plugins/Experimental/NeuralNetworkInference/Source/NeuralNetworkInference/Public/NeuralNetwork.h
--- a/Engine/Plugins/Experimental/NeuralNetworkInference/Source/NeuralNetworkInference/Public/NeuralNetwork.h
+++ b/Engine/Plugins/Experimental/NeuralNetworkInference/Source/NeuralNetworkInference/Public/NeuralNetwork.h
@@ -125,7 +125,7 @@ public:
 	ENeuralDeviceType GetOutputDeviceType() const;
 	void SetDeviceType(const ENeuralDeviceType InDeviceType, const ENeuralDeviceType InInputDeviceType = ENeuralDeviceType::CPU,
 		const ENeuralDeviceType InOutputDeviceType = ENeuralDeviceType::CPU);
-	
+
 	/**
 	 * Whether GPU execution (i.e., SetDeviceType(ENeuralDeviceType::GPU)) is supported for this platform:
 	 * - On Windows:
@@ -232,51 +232,57 @@ public:
 	 */
 	void Run();
 
-// CONTINUE HERE
 	/**
-	 * Stats functions:
+	 * Statistics-related functions:
 	 * - GetLastRunTimeMSec will provide the last inference time measured milliseconds.
-	 * - GetInferenceStats returns inference time statistics. (NumberSamples, Average, StdDev, Min, Max statistics measured in milliseconds)
-	 * - GetInputMemoryTransferStats returns Input Memory Transfer statistics. (NumberSamples, Average, StdDev, Min, Max statistics measured in
-	 *   milliseconds)
+	 * - GetRunStatistics returns a FNeuralStatistics with the time statistics in milliseconds (NumberSamples, Average, StdDev, Min, Max).
+	 * - GetInputMemoryTransferStats returns a FNeuralStatistics with the input memory transfer statistics in milliseconds.
+	 * @see FNeuralStatistics for more details.
 	 */
 	float GetLastRunTimeMSec() const;
-	FNeuralStatistics GetInferenceStats() const;
+	FNeuralStatistics GetRunStatistics() const;
 	FNeuralStatistics GetInputMemoryTransferStats() const;
 	void ResetStats();
 
 protected:
 	/**
-	 * Whether Run() will use CPU or GPU acceleration hardware.
-	 * If SetDeviceType() is never called, the default device (EDeviceType::CPU) will be used.
+	 * The neural device type of the network. It defines whether the network will use CPU or GPU acceleration hardware during inference (on Run).
+	 * If SetDeviceType() is never called, the default device (EDeviceType::GPU) will be used.
+	 * @see ENeuralDeviceType, InputDeviceType, OutputDeviceType for more details.
 	 */
 	UPROPERTY(EditAnywhere, Category = "Neural Network Inference")
 	ENeuralDeviceType DeviceType;
-	
+
 	/**
-	 * If DeviceType == CPU, InputDeviceType and OutputDeviceType must also be set to CPU.
-	 * If DeviceType == GPU:
-	 *  - InputDeviceType: Whether Run() will expect the input data in CPU (Run will upload the memory to the GPU) or GPU (no upload needed).
-	 *  - OutputDeviceType: Whether Run() will return output data in CPU (Run will download the memory to the CPU) or GPU (no download needed).
+	 * It defines whether Run() will expect the input data in CPU memory (Run will upload the memory to the GPU) or GPU memory (no upload needed).
+	 * If DeviceType == CPU, InputDeviceType and OutputDeviceType values are ignored and assumed to be set to CPU.
+	 * @see ENeuralDeviceType, DeviceType, OutputDeviceType for more details.
 	 */
 	UPROPERTY(EditAnywhere, Category = "Neural Network Inference")
 	ENeuralDeviceType InputDeviceType;
-	
+
+	/**
+	 * It defines whether Run() will return the output data in CPU memory (Run will download the memory to the CPU) or GPU (no download needed).
+	 * If DeviceType == CPU, InputDeviceType and OutputDeviceType values are ignored and assumed to be set to CPU.
+	 * @see ENeuralDeviceType, DeviceType, InputDeviceType for more details.
+	 */
 	UPROPERTY(EditAnywhere, Category = "Neural Network Inference")
 	ENeuralDeviceType OutputDeviceType;
-	
+
 	/**
-	 * SynchronousMode defines whether UNeuralNetwork::Run() will block the thread until completed (Synchronous), or whether it will run on a
-	 * background thread, not blocking the calling thread (Asynchronous).
-	 * If asynchronous, ThreadModeDelegateForAsyncRunCompleted will define whether the callback delegate is called from the game thread (highly recommended) or from
-	 * any available thread (not fully thread safe).
-	 * - If ThreadModeDelegateForAsyncRunCompleted == ENeuralSynchronousMode::Asynchronous, the FOnAsyncRunCompleted delegate could be triggered from any thread.
-	 * - If ThreadModeDelegateForAsyncRunCompleted == ENeuralSynchronousMode::Synchronous, UNeuralNetwork::Run() will block the calling thread until completed, so a
-	 *   callback delegate is not required.
-	 * @see ENeuralSynchronousMode, ENeuralThreadMode for more details.
+	 * It defines whether UNeuralNetwork::Run() will block the thread until completed (Synchronous), or whether it will run on a background thread,
+	 * not blocking the calling thread (Asynchronous).
 	 */
 	UPROPERTY(Transient, EditAnywhere, Category = "Neural Network Inference")
 	ENeuralSynchronousMode SynchronousMode;
+
+	/**
+	 * If SynchronousMode is Asynchronous, this variable will define whether the callback delegate is called from the game thread (highly
+	 * recommended) or from any available thread (not fully thread safe).
+	 * - If this variable is set to ENeuralThreadMode::GameThread, the FOnAsyncRunCompleted delegate will be triggered from the main thread only.
+	 * - If this variable is set to ENeuralThreadMode::AnyThread, the FOnAsyncRunCompleted delegate could be triggered from any thread.
+	 * @see SynchronousMode for more details.
+	 */
 	UPROPERTY(Transient, EditAnywhere, Category = "Neural Network Inference")
 	ENeuralThreadMode ThreadModeDelegateForAsyncRunCompleted;
 
@@ -287,19 +293,28 @@ protected:
 	FString ModelFullFilePath;
 
 private:
+	/**
+	 * Whether this UNeuralNetwork instance has loaded a valid network model already, i.e., whether Load() was called and returned true.
+	 */
 	UPROPERTY(Transient, VisibleAnywhere, Category = "Neural Network Inference")
 	bool bIsLoaded;
 
+	/**
+	 * A buffer of memory representing the ONNX model. It is equivalent to the serialization of the ONNX file.
+	 * During Load(InModelFilePath), the ONNX file is read and stored as this buffer of raw bytes (TArray<uint8>).
+	 */
 	UPROPERTY()
 	TArray<uint8> ModelReadFromFileInBytes;
 
-	/** Whether some of the FNeuralTensor of InputTensor have flexible/variable dimensions. */
+	/**
+	 * Whether some of the FNeuralTensor of InputTensor have flexible/variable dimensions. E.g., this is useful for variable batch size.
+	 */
 	UPROPERTY(Transient, VisibleAnywhere, Category = "Neural Network Inference")
 	TArray<bool> AreInputTensorSizesVariable;
 
 	/**
 	 * Mutex to avoid issues or crashes due to the asynchronous Run() being run at the same time than any other non-const function.
-	 * @see UNeuralNetwork::Run().
+	 * @see UNeuralNetwork::Run() implementation for more details.
 	 */
 	FCriticalSection ResoucesCriticalSection;
 
@@ -309,35 +324,30 @@ private:
 	FOnAsyncRunCompleted OnAsyncRunCompletedDelegate;
 
 	/**
-	 * Stats-related members.
+	 * Statistics-related members used for GetLastRunTimeMSec(), GetRunStatistics(), GetInputMemoryTransferStats(), ResetStats().
 	 */
-	FNeuralStatisticsEstimator ComputeStatisticsEstimator;
+	FNeuralStatisticsEstimator RunStatisticsEstimator;
 	FNeuralStatisticsEstimator InputTransferStatisticsEstimator;
 
 	/**
-	 * Struct pointer containing the UE-and-ORT-based back end implementation.
-	 * PIMPL idiom to minimize memory when not using this back end and to hide 3rd party dependencies.
+	 * Struct pointers containing the UE-and-ORT-based (ImplBackEndUEAndORT) and only-UE-based (ImplBackEndUEOnly) back end implementations.
+	 * PIMPL idiom to minimize memory when not using each back end (only 1 is used at the time) and to hide 3rd party dependencies.
 	 * http://www.cppsamples.com/common-tasks/pimpl.html
 	 */
 	struct FImplBackEndUEAndORT;
 	TSharedPtr<FImplBackEndUEAndORT> ImplBackEndUEAndORT;
-
-	/**
-	 * Struct pointer containing the only-UE-based back end implementation.
-	 * PIMPL idiom to minimize memory when not using this back end.
-	 * http://www.cppsamples.com/common-tasks/pimpl.html
-	 */
 	struct FImplBackEndUEOnly;
 	TSharedPtr<FImplBackEndUEOnly> ImplBackEndUEOnly;
 
 	/**
-	 * It loads the desired network graph definition and weights internally saved on this UNeuralNetwork instance.
+	 * It loads the existing network graph definition and weights from the members of this UNeuralNetwork instance.
 	 * @return Whether the network was successfully loaded.
 	 */
 	bool Load();
 
 	/**
-	 * GetInputTensorMutable() and GetOutputTensorMutable() are the private and mutable version of GetInputTensor() and GetOutputTensor(), respectively.
+	 * GetInputTensorMutable() and GetOutputTensorMutable() are the private and mutable version of GetInputTensor() and GetOutputTensor(),
+	 * respectively.
 	 * Most FNeuralTensor properties are considered constant and fixed after calling UNeuralNetwork::Load(), other than the actual data values
 	 * (i.e., UnderlyingUInt8ArrayData) and some GPU properties. Thus, use these 2 functions with extreme precaution and be sure not to modify any
 	 * other FNeuralNetwork properties.
@@ -346,30 +356,30 @@ private:
 	FNeuralTensor& GetInputTensorMutable(const int32 InTensorIndex = 0);
 	FNeuralTensor& GetOutputTensorMutable(const int32 InTensorIndex = 0);
 
-public:
+public: // It should really say "private-for-user, public-for-auxiliary-NNI-classes"
 	/**
 	 * Internal function not needed by the user.
 	 * Used to create custom networks without an ONNX file for QA testing in FOperatorTester::TestOperator().
 	 */
 	bool Load(TArray<FNeuralTensor>& InTensors, const TArray<FNeuralTensor*>& InInputTensors, const TArray<FNeuralTensor*>& InOutputTensors,
 		const TArray<TSharedPtr<class FNeuralOperator>>& InOperators);
-	
+
 	/**
 	 * Internal enum class that should not be used by the user.
-	 * Whether UNeuralNetwork will use the highly optimized UnrealEngine-and-ONNXRuntime-based back end (UEAndORT) or the less optimized but fully
-	 * cross platform only-UE one (UEOnly).
-	 * We recommend using Auto, which will find and use the optimal back end for each platform.
+	 * It enumerates the different types of back ends. Use Auto, which will find and use the optimal back end for each platform.
 	 */
 	enum class ENeuralBackEnd : uint8
 	{
 		/**
-		 * UnrealEngine-and-ONNXRuntime-accelerated back end, ideal for those platforms that support it. WITH_UE_AND_ORT_SUPPORT is the C++ define
-		 * macro that checks whether UEAndORT support exists for the current platform.
+		 * Highly optimized UnrealEngine-and-ONNXRuntime-based back end, ideal for those platforms that support it. WITH_UE_AND_ORT_SUPPORT is the
+		 * C++ define macro that checks whether UEAndORT support exists for the current platform.
 		 */
 		UEAndORT,
-		/** It might be slower than the UEAndORT back end, but it will compile in all platforms and OSs compatible with Unreal Engine. */
+		/**
+		 * Slower than the UEAndORT back end maintained for now for historical reasons. There is no benefit in using this back end.
+		 */
 		UEOnly,
-		/** Recommended value. It will use the efficient UEAndORT if supported by the platform, and fall back to UEOnly otherwise. */
+		/** Recommended value. It will use the efficient UEAndORT if supported by the platform, and try to fall back to UEOnly otherwise. */
 		Auto
 	};
 
@@ -384,14 +394,17 @@ public:
 	 * SetBackEnd() is NOT THREAD SAFE, make sure that other non-const functions such as Run() are not running when SetBackEnd is called.
 	 * @see ENeuralBackEnd for more details.
 	 */
+	UE_DEPRECATED(5.0, "Do not use, ENeuralBackEnd will be removed in future versions and only UEAndORT back end will be supported.")
 	ENeuralBackEnd GetBackEnd() const;
+	UE_DEPRECATED(5.0, "Do not use, ENeuralBackEnd will be removed in future versions and only UEAndORT back end will be supported.")
 	ENeuralBackEnd GetBackEndForCurrentPlatform() const;
+	UE_DEPRECATED(5.0, "Do not use, ENeuralBackEnd will be removed in future versions and only UEAndORT back end will be supported.")
 	bool SetBackEnd(const ENeuralBackEnd InBackEnd);
 
 #if WITH_EDITOR
 	/**
 	 * Internal and Editor-only functions not needed by the user.
-	 * Importing data and options used for loading the network.
+	 * They provide importing data and options used for loading the network.
 	 */
 	TObjectPtr<class UAssetImportData> GetAssetImportData() const;
 	TObjectPtr<class UAssetImportData> GetAndMaybeCreateAssetImportData();
@@ -399,13 +412,12 @@ public:
 
 private:
 	/**
-	 * Internal variable that should not be used by the user.
-	 * @see ENeuralBackEnd for more details.
+	 * The back end to be used (it could and should be set to Auto). @see ENeuralBackEnd for more details.
 	 */
 	ENeuralBackEnd BackEnd;
 
 	/**
-	 * Internal variable that should not be used by the user.
+	 * The actual back end currently being used (it cannot be set to Auto).
 	 * If BackEnd != Auto, BackEndForCurrentPlatform will be equal to BackEnd.
 	 * Otherwise, BackEndForCurrentPlatform will be set to the optimal BackEnd given the current platform.
 	 * @see ENeuralBackEnd for more details.
diff --git a/Engine/Plugins/Experimental/NeuralNetworkInference/Source/Private/NeuralNetworkInferenceQA/Private/ModelUnitTester.cpp b/Engine/Plugins/Experimental/NeuralNetworkInference/Source/Private/NeuralNetworkInferenceQA/Private/ModelUnitTester.cpp
--- a/Engine/Plugins/Experimental/NeuralNetworkInference/Source/Private/NeuralNetworkInferenceQA/Private/ModelUnitTester.cpp
+++ b/Engine/Plugins/Experimental/NeuralNetworkInference/Source/Private/NeuralNetworkInferenceQA/Private/ModelUnitTester.cpp
@@ -59,7 +59,7 @@ FNNIUnitTesterTimeData FModelUnitTester_GetTimeInformation(UNeuralNetwork* InOut
 		}
 	}
 	// Return NetworkTimeData
-	return FNNIUnitTesterTimeData(InOutNetwork->GetInferenceStats(), InOutNetwork->GetInputMemoryTransferStats(), OutCopyingStats.GetStats());
+	return FNNIUnitTesterTimeData(InOutNetwork->GetRunStatistics(), InOutNetwork->GetInputMemoryTransferStats(), OutCopyingStats.GetStats());
 }
 
 
@@ -232,7 +232,9 @@ bool FModelUnitTester::ModelLoadAccuracyAndSpeedTests(const FString& InProjectCo
 		bDidGlobalTestPassed &= ModelSpeedTest(UAssetModelFilePath, ENeuralDeviceType::CPU, UNeuralNetwork::ENeuralBackEnd::UEAndORT,
 			InCPURepetitionsForUEAndORTBackEnd[ModelIndex]);
 		UNeuralNetwork* Network = NetworkUassetLoadTest(UAssetModelFilePath);
+PRAGMA_DISABLE_DEPRECATION_WARNINGS
 		if (Network->GetBackEndForCurrentPlatform() != UNeuralNetwork::ENeuralBackEnd::UEAndORT)
+PRAGMA_ENABLE_DEPRECATION_WARNINGS
 		{
 			UE_LOG(LogNeuralNetworkInferenceQA, Warning, TEXT("-------------------- Default UAsset BackEnd should be UEAndORT."));
 			return false;
@@ -349,14 +351,18 @@ bool FModelUnitTester::ModelAccuracyTest(UNeuralNetwork* InOutNetwork, const ENe
 	const ENeuralDeviceType OriginalOutputDeviceType = InOutNetwork->GetOutputDeviceType();
 	const ENeuralSynchronousMode OriginalSynchronousMode = InOutNetwork->GetSynchronousMode();
 	const ENeuralThreadMode OriginalThreadModeDelegateForAsyncRunCompleted = InOutNetwork->GetThreadModeDelegateForAsyncRunCompleted();
+PRAGMA_DISABLE_DEPRECATION_WARNINGS
 	const UNeuralNetwork::ENeuralBackEnd OriginalBackEnd = InOutNetwork->GetBackEnd();
+PRAGMA_ENABLE_DEPRECATION_WARNINGS
 	
 	// Set (a)synchronous Mode
 	InOutNetwork->SetSynchronousMode(InSynchronousMode);
 	InOutNetwork->SetThreadModeDelegateForAsyncRunCompleted(ENeuralThreadMode::AnyThread);
 
 	// Set back end
+PRAGMA_DISABLE_DEPRECATION_WARNINGS
 	if (!InOutNetwork->SetBackEnd(InBackEnd))
+PRAGMA_ENABLE_DEPRECATION_WARNINGS
 	{
 		UE_LOG(LogNeuralNetworkInferenceQA, Warning, TEXT("Backend %s is disabled."), *GetBackEndString(InBackEnd));
 		return false;
@@ -459,7 +465,9 @@ bool FModelUnitTester::ModelAccuracyTest(UNeuralNetwork* InOutNetwork, const ENe
 		/*OutputDeviceType*/OriginalOutputDeviceType);
 	InOutNetwork->SetSynchronousMode(OriginalSynchronousMode);
 	InOutNetwork->SetThreadModeDelegateForAsyncRunCompleted(OriginalThreadModeDelegateForAsyncRunCompleted);
+PRAGMA_DISABLE_DEPRECATION_WARNINGS
 	InOutNetwork->SetBackEnd(OriginalBackEnd);
+PRAGMA_ENABLE_DEPRECATION_WARNINGS
 	// Test successful
 	return bDidGlobalTestPassed;
 }
@@ -532,9 +540,13 @@ bool FModelUnitTester::ModelSpeedTest(const FString& InUAssetPath, const ENeural
 	}
 	// Save original network state
 	const ENeuralDeviceType OriginalDeviceType = InOutNetwork->GetDeviceType();
+PRAGMA_DISABLE_DEPRECATION_WARNINGS
 	const UNeuralNetwork::ENeuralBackEnd OriginalBackEnd = InOutNetwork->GetBackEnd();
+PRAGMA_ENABLE_DEPRECATION_WARNINGS
 	// Set desired back end
+PRAGMA_DISABLE_DEPRECATION_WARNINGS
 	if (!InOutNetwork->SetBackEnd(InBackEnd))
+PRAGMA_ENABLE_DEPRECATION_WARNINGS
 	{
 		UE_LOG(LogNeuralNetworkInferenceQA, Warning, TEXT("Backend %s is disabled."), *GetBackEndString(InBackEnd));
 		return false;
@@ -563,7 +575,9 @@ bool FModelUnitTester::ModelSpeedTest(const FString& InUAssetPath, const ENeural
 
 	// Reset to original network state
 	InOutNetwork->SetDeviceType(OriginalDeviceType);
+PRAGMA_DISABLE_DEPRECATION_WARNINGS
 	InOutNetwork->SetBackEnd(OriginalBackEnd);
+PRAGMA_ENABLE_DEPRECATION_WARNINGS
 	// Test successful
 	return true;
 }
