Fixed issue of using wrong logging path for ORT.

#rnx

diff --git a/Engine/Plugins/Experimental/NeuralNetworkInference/Source/NeuralNetworkInference/Private/NeuralNetworkImplBackEndUEAndORT.cpp b/Engine/Plugins/Experimental/NeuralNetworkInference/Source/NeuralNetworkInference/Private/NeuralNetworkImplBackEndUEAndORT.cpp
--- a/Engine/Plugins/Experimental/NeuralNetworkInference/Source/NeuralNetworkInference/Private/NeuralNetworkImplBackEndUEAndORT.cpp
+++ b/Engine/Plugins/Experimental/NeuralNetworkInference/Source/NeuralNetworkInference/Private/NeuralNetworkImplBackEndUEAndORT.cpp
@@ -416,7 +416,7 @@ bool UNeuralNetwork::FImplBackEndUEAndORT::Load(TSharedPtr<FImplBackEndUEAndORT>
 #endif //WITH_EDITOR
 
 #else //WITH_UE_AND_ORT_SUPPORT
-	UE_LOG(LogNeuralNetworkInference, Warning, TEXT("FImplBackEndUEAndORT::Load(): Platform or Operating System not suported yet for UEAndORT"
+	UE_LOG(LogNeuralNetworkInference, Warning, TEXT("FImplBackEndUEAndORT::Load(): Platform or Operating System not supported yet for UEAndORT"
 		" BackEnd. Set BackEnd to ENeuralBackEnd::Auto (recommended) or ENeuralBackEnd::UEOnly for this platform."));
 	return false;
 #endif //WITH_UE_AND_ORT_SUPPORT
@@ -484,9 +484,9 @@ bool UNeuralNetwork::FImplBackEndUEAndORT::InitializedAndConfigureMembers(TShare
 		InOutImplBackEndUEAndORT = MakeShared<FImplBackEndUEAndORT>(InOutOnAsyncRunCompletedDelegate, InOutDelegateThreadMode, InOutResoucesCriticalSection);
 
 		// Set up ORT and create an environment
-		const auto ModelFullFilePathAnsiString = StringCast<ANSICHAR>(*InModelFullFilePath);
-		// @todo: ModelFullFilePathAnsiString -> I thought any unique string would work, but it might be output logging file, so it has to be a non-existing file!
-		InOutImplBackEndUEAndORT->Environment = MakeUnique<Ort::Env>(ORT_LOGGING_LEVEL_WARNING, ModelFullFilePathAnsiString.Get());
+		const FString LoggingFullFilePath = InModelFullFilePath + TEXT(".txt");
+		const auto LoggingFullFilePathAnsiString = StringCast<ANSICHAR>(*LoggingFullFilePath);
+		InOutImplBackEndUEAndORT->Environment = MakeUnique<Ort::Env>(ORT_LOGGING_LEVEL_WARNING, LoggingFullFilePathAnsiString.Get());
 
 		InOutImplBackEndUEAndORT->Allocator = MakeUnique<Ort::AllocatorWithDefaultOptions>();
 
