Improved doc.

#rnx

diff --git a/Engine/Plugins/Experimental/NeuralNetworkInference/Source/NeuralNetworkInference/Public/NeuralNetwork.h b/Engine/Plugins/Experimental/NeuralNetworkInference/Source/NeuralNetworkInference/Public/NeuralNetwork.h
--- a/Engine/Plugins/Experimental/NeuralNetworkInference/Source/NeuralNetworkInference/Public/NeuralNetwork.h
+++ b/Engine/Plugins/Experimental/NeuralNetworkInference/Source/NeuralNetworkInference/Public/NeuralNetwork.h
@@ -21,7 +21,7 @@
  * The following code snippets show the UNeuralNetwork basics (reading a ONNX model and running inference on it). For more detailed examples, see
  * {UE5}/Samples/MachineLearning/NNI.
  *
- * 1.a. Creating a new UNeuralNetwork and loading a network from an ONNX file:
+ * 1a Creating a new UNeuralNetwork and loading a network from an ONNX file
  *		// Create the UNeuralNetwork object
  *		UNeuralNetwork* Network = NewObject<UNeuralNetwork>((UObject*)GetTransientPackage(), UNeuralNetwork::StaticClass());
  *		// Load the ONNX model and set the device (CPU/GPU)
@@ -35,42 +35,47 @@
  *			// Option b) Set to CPU
  *			Network->SetDeviceType(ENeuralDeviceType::CPU);
  *		}
+ *		// Check that the network was successfully loaded
+ *		else
+ *		{
+ *			UE_LOG(LogTemp, Warning, TEXT("UNeuralNetwork could not loaded from %s."), *ONNXModelFilePath);
+ *		}
  *
- * 1.b. Loading a UNeuralNetwork from a previously-created UNeuralNetwork UAsset:
+ * 1b Loading a UNeuralNetwork from a previously-created UNeuralNetwork UAsset
  *		// Load the UNeuralNetwork object from a UNeuralNetwork UAsset
  *		const FString NetworkUAssetFilePath = TEXT("'/Game/Models/ExampleNetwork/ExampleNetwork.ExampleNetwork'");
  *		UNeuralNetwork* Network = LoadObject<UNeuralNetwork>((UObject*)GetTransientPackage(), *NetworkUAssetFilePath);
  *		// Check that the network was successfully loaded
- *		check(Network->IsLoaded());
- *      // Optionally set to CPU/GPU mode. This step is optional, if not called, it will use the device type that was saved on the loaded UAsset
- *      // Network->SetDeviceType(ENeuralDeviceType::CPU);
+ *		ensureMsgf(Network->IsLoaded(), TEXT("UNeuralNetwork could not loaded from %s."), *NetworkUAssetFilePath);
+ *		// Optionally set to CPU/GPU mode. This step is optional, if not called, it will use the device type that was saved on the loaded UAsset
+ *		// Network->SetDeviceType(ENeuralDeviceType::CPU);
  *
- * 2.a. Running inference:
+ * 2a Running inference
  *		// Fill input neural tensor
  *		TArray<float> InArray; // Assumed initialized with data and that InArray.Num() == Network->Num()
  *		Network->SetInputFromArrayCopy(InArray); // Equivalent: Network->SetInputFromVoidPointerCopy(InArray.GetData());
- *		UE_LOG(LogNeuralNetworkInference, Display, TEXT("Input tensor: %s."), *Network->GetInputTensor().ToString());
+ *		UE_LOG(LogTemp, Display, TEXT("Input tensor: %s."), *Network->GetInputTensor().ToString());
  *		// Run UNeuralNetwork
  *		Network->Run();
  *		// Read and print OutputTensor
  *		const FNeuralTensor& OutputTensor = Network->GetOutputTensor();
- *		UE_LOG(LogNeuralNetworkInference, Display, TEXT("Output tensor: %s."), *OutputTensor.ToString());
+ *		UE_LOG(LogTemp, Display, TEXT("Output tensor: %s."), *OutputTensor.ToString());
  *
- * 2.b. Running inference more efficient alternative - Filling the input tensor without a TArray-to-FNeuralTensor copy:
+ * 2b Running inference more efficiently - Filling the input tensor without a TArray-to-FNeuralTensor copy
  *		// Obtain input tensor pointer
- *		float* InputDataPointer = Network->GetInputDataPointerMutable<float>();
+ *		float* InputDataPointer = (float*)Network->GetInputDataPointerMutable();
  *		// Fill InputDataPointer
  *		for (int64 Index = 0; Index < Network->GetInputTensor().Num(); ++Index)
  *			InputDataPointer[Index] = ...; // Assumed some preprocessing or otherwise simply use Memcpy
  *
- * 3. Networks with multiple input/output tensors:
- *      - Multiple inputs: Add InTensorIndex to SetInputFromArrayCopy, SetInputFromVoidPointerCopy, GetInputTensor or GetInputDataPointerMutable
- *        in the examples above:
- *		    Network->SetInputFromArrayCopy(InputArray0, 0); // Equivalent: Network->SetInputFromVoidPointerCopy(InputArray0.GetData(), 0);
- *		    Network->SetInputFromArrayCopy(InputArray1, 1); // Equivalent: Network->SetInputFromVoidPointerCopy(InputArray1.GetData(), 1);
- *      - Multiple outputs: Add InTensorIndex to GetOutputTensor(InTensorIndex) in the examples above or use GetOutputTensors() instead.
- *		    const FNeuralTensor& OutputTensor0 = Network->GetOutputTensor(0);
- *		    const FNeuralTensor& OutputTensor1 = Network->GetOutputTensor(1);
+ * 3 Networks with multiple input/output tensors
+ *		- Multiple inputs: Add InTensorIndex to SetInputFromArrayCopy, SetInputFromVoidPointerCopy, GetInputTensor or GetInputDataPointerMutable
+ *		  in the examples above:
+ *			Network->SetInputFromArrayCopy(InputArray0, 0); // Equivalent: Network->SetInputFromVoidPointerCopy(InputArray0.GetData(), 0);
+ *			Network->SetInputFromArrayCopy(InputArray1, 1); // Equivalent: Network->SetInputFromVoidPointerCopy(InputArray1.GetData(), 1);
+ *		- Multiple outputs: Add InTensorIndex to GetOutputTensor(InTensorIndex) in the examples above or use GetOutputTensors() instead.
+ *			const FNeuralTensor& OutputTensor0 = Network->GetOutputTensor(0);
+ *			const FNeuralTensor& OutputTensor1 = Network->GetOutputTensor(1);
  */
 UCLASS(BlueprintType)
 class NEURALNETWORKINFERENCE_API UNeuralNetwork : public UObject
