Improved doc.

#rnx

diff --git a/Engine/Plugins/Experimental/NeuralNetworkInference/Source/NeuralNetworkInference/Public/NeuralNetwork.h b/Engine/Plugins/Experimental/NeuralNetworkInference/Source/NeuralNetworkInference/Public/NeuralNetwork.h
--- a/Engine/Plugins/Experimental/NeuralNetworkInference/Source/NeuralNetworkInference/Public/NeuralNetwork.h
+++ b/Engine/Plugins/Experimental/NeuralNetworkInference/Source/NeuralNetworkInference/Public/NeuralNetwork.h
@@ -21,7 +21,7 @@
  * The following code snippets show the UNeuralNetwork basics (reading a ONNX model and running inference on it). For more detailed examples, see
  * {UE5}/Samples/MachineLearning/NNI.
  *
- * 1.a. Creating a new UNeuralNetwork and loading a network from an ONNX file:
+ * 1a Creating a new UNeuralNetwork and loading a network from an ONNX file
  *		// Create the UNeuralNetwork object
  *		UNeuralNetwork* Network = NewObject<UNeuralNetwork>((UObject*)GetTransientPackage(), UNeuralNetwork::StaticClass());
  *		// Load the ONNX model and set the device (CPU/GPU)
@@ -35,42 +35,47 @@
  *			// Option b) Set to CPU
  *			Network->SetDeviceType(ENeuralDeviceType::CPU);
  *		}
+ *		// Check that the network was successfully loaded
+ *		else
+ *		{
+ *			UE_LOG(LogTemp, Warning, TEXT("UNeuralNetwork could not loaded from %s."), *ONNXModelFilePath);
+ *		}
  *
- * 1.b. Loading a UNeuralNetwork from a previously-created UNeuralNetwork UAsset:
+ * 1b Loading a UNeuralNetwork from a previously-created UNeuralNetwork UAsset
  *		// Load the UNeuralNetwork object from a UNeuralNetwork UAsset
  *		const FString NetworkUAssetFilePath = TEXT("'/Game/Models/ExampleNetwork/ExampleNetwork.ExampleNetwork'");
  *		UNeuralNetwork* Network = LoadObject<UNeuralNetwork>((UObject*)GetTransientPackage(), *NetworkUAssetFilePath);
  *		// Check that the network was successfully loaded
- *		check(Network->IsLoaded());
- *      // Optionally set to CPU/GPU mode. This step is optional, if not called, it will use the device type that was saved on the loaded UAsset
- *      // Network->SetDeviceType(ENeuralDeviceType::CPU);
+ *		ensureMsgf(Network->IsLoaded(), TEXT("UNeuralNetwork could not loaded from %s."), *NetworkUAssetFilePath);
+ *		// Optionally set to CPU/GPU mode. This step is optional, if not called, it will use the device type that was saved on the loaded UAsset
+ *		// Network->SetDeviceType(ENeuralDeviceType::CPU);
  *
- * 2.a. Running inference:
+ * 2a Running inference
  *		// Fill input neural tensor
  *		TArray<float> InArray; // Assumed initialized with data and that InArray.Num() == Network->Num()
  *		Network->SetInputFromArrayCopy(InArray); // Equivalent: Network->SetInputFromVoidPointerCopy(InArray.GetData());
- *		UE_LOG(LogNeuralNetworkInference, Display, TEXT("Input tensor: %s."), *Network->GetInputTensor().ToString());
+ *		UE_LOG(LogTemp, Display, TEXT("Input tensor: %s."), *Network->GetInputTensor().ToString());
  *		// Run UNeuralNetwork
  *		Network->Run();
  *		// Read and print OutputTensor
  *		const FNeuralTensor& OutputTensor = Network->GetOutputTensor();
- *		UE_LOG(LogNeuralNetworkInference, Display, TEXT("Output tensor: %s."), *OutputTensor.ToString());
+ *		UE_LOG(LogTemp, Display, TEXT("Output tensor: %s."), *OutputTensor.ToString());
  *
- * 2.b. Running inference more efficient alternative - Filling the input tensor without a TArray-to-FNeuralTensor copy:
+ * 2b Running inference more efficiently - Filling the input tensor without a TArray-to-FNeuralTensor copy
  *		// Obtain input tensor pointer
- *		float* InputDataPointer = Network->GetInputDataPointerMutable<float>();
+ *		float* InputDataPointer = (float*)Network->GetInputDataPointerMutable();
  *		// Fill InputDataPointer
  *		for (int64 Index = 0; Index < Network->GetInputTensor().Num(); ++Index)
  *			InputDataPointer[Index] = ...; // Assumed some preprocessing or otherwise simply use Memcpy
  *
- * 3. Networks with multiple input/output tensors:
- *      - Multiple inputs: Add InTensorIndex to SetInputFromArrayCopy, SetInputFromVoidPointerCopy, GetInputTensor or GetInputDataPointerMutable
- *        in the examples above:
- *		    Network->SetInputFromArrayCopy(InputArray0, 0); // Equivalent: Network->SetInputFromVoidPointerCopy(InputArray0.GetData(), 0);
- *		    Network->SetInputFromArrayCopy(InputArray1, 1); // Equivalent: Network->SetInputFromVoidPointerCopy(InputArray1.GetData(), 1);
- *      - Multiple outputs: Add InTensorIndex to GetOutputTensor(InTensorIndex) in the examples above or use GetOutputTensors() instead.
- *		    const FNeuralTensor& OutputTensor0 = Network->GetOutputTensor(0);
- *		    const FNeuralTensor& OutputTensor1 = Network->GetOutputTensor(1);
+ * 3 Networks with multiple input/output tensors
+ *		- Multiple inputs: Add InTensorIndex to SetInputFromArrayCopy, SetInputFromVoidPointerCopy, GetInputTensor or GetInputDataPointerMutable
+ *		  in the examples above:
+ *			Network->SetInputFromArrayCopy(InputArray0, 0); // Equivalent: Network->SetInputFromVoidPointerCopy(InputArray0.GetData(), 0);
+ *			Network->SetInputFromArrayCopy(InputArray1, 1); // Equivalent: Network->SetInputFromVoidPointerCopy(InputArray1.GetData(), 1);
+ *		- Multiple outputs: Add InTensorIndex to GetOutputTensor(InTensorIndex) in the examples above or use GetOutputTensors() instead.
+ *			const FNeuralTensor& OutputTensor0 = Network->GetOutputTensor(0);
+ *			const FNeuralTensor& OutputTensor1 = Network->GetOutputTensor(1);
  */
 UCLASS(BlueprintType)
 class NEURALNETWORKINFERENCE_API UNeuralNetwork : public UObject
@@ -125,7 +130,7 @@ public:
 	ENeuralDeviceType GetOutputDeviceType() const;
 	void SetDeviceType(const ENeuralDeviceType InDeviceType, const ENeuralDeviceType InInputDeviceType = ENeuralDeviceType::CPU,
 		const ENeuralDeviceType InOutputDeviceType = ENeuralDeviceType::CPU);
-	
+
 	/**
 	 * Whether GPU execution (i.e., SetDeviceType(ENeuralDeviceType::GPU)) is supported for this platform:
 	 * - On Windows:
@@ -232,51 +237,57 @@ public:
 	 */
 	void Run();
 
-// CONTINUE HERE
 	/**
-	 * Stats functions:
+	 * Statistics-related functions:
 	 * - GetLastRunTimeMSec will provide the last inference time measured milliseconds.
-	 * - GetInferenceStats returns inference time statistics. (NumberSamples, Average, StdDev, Min, Max statistics measured in milliseconds)
-	 * - GetInputMemoryTransferStats returns Input Memory Transfer statistics. (NumberSamples, Average, StdDev, Min, Max statistics measured in
-	 *   milliseconds)
+	 * - GetRunStatistics returns a FNeuralStatistics with the time statistics in milliseconds (NumberSamples, Average, StdDev, Min, Max).
+	 * - GetInputMemoryTransferStats returns a FNeuralStatistics with the input memory transfer statistics in milliseconds.
+	 * @see FNeuralStatistics for more details.
 	 */
 	float GetLastRunTimeMSec() const;
-	FNeuralStatistics GetInferenceStats() const;
+	FNeuralStatistics GetRunStatistics() const;
 	FNeuralStatistics GetInputMemoryTransferStats() const;
 	void ResetStats();
 
 protected:
 	/**
-	 * Whether Run() will use CPU or GPU acceleration hardware.
-	 * If SetDeviceType() is never called, the default device (EDeviceType::CPU) will be used.
+	 * The neural device type of the network. It defines whether the network will use CPU or GPU acceleration hardware during inference (on Run).
+	 * If SetDeviceType() is never called, the default device (EDeviceType::GPU) will be used.
+	 * @see ENeuralDeviceType, InputDeviceType, OutputDeviceType for more details.
 	 */
 	UPROPERTY(EditAnywhere, Category = "Neural Network Inference")
 	ENeuralDeviceType DeviceType;
-	
+
 	/**
-	 * If DeviceType == CPU, InputDeviceType and OutputDeviceType must also be set to CPU.
-	 * If DeviceType == GPU:
-	 *  - InputDeviceType: Whether Run() will expect the input data in CPU (Run will upload the memory to the GPU) or GPU (no upload needed).
-	 *  - OutputDeviceType: Whether Run() will return output data in CPU (Run will download the memory to the CPU) or GPU (no download needed).
+	 * It defines whether Run() will expect the input data in CPU memory (Run will upload the memory to the GPU) or GPU memory (no upload needed).
+	 * If DeviceType == CPU, InputDeviceType and OutputDeviceType values are ignored and assumed to be set to CPU.
+	 * @see ENeuralDeviceType, DeviceType, OutputDeviceType for more details.
 	 */
 	UPROPERTY(EditAnywhere, Category = "Neural Network Inference")
 	ENeuralDeviceType InputDeviceType;
-	
+
+	/**
+	 * It defines whether Run() will return the output data in CPU memory (Run will download the memory to the CPU) or GPU (no download needed).
+	 * If DeviceType == CPU, InputDeviceType and OutputDeviceType values are ignored and assumed to be set to CPU.
+	 * @see ENeuralDeviceType, DeviceType, InputDeviceType for more details.
+	 */
 	UPROPERTY(EditAnywhere, Category = "Neural Network Inference")
 	ENeuralDeviceType OutputDeviceType;
-	
+
 	/**
-	 * SynchronousMode defines whether UNeuralNetwork::Run() will block the thread until completed (Synchronous), or whether it will run on a
-	 * background thread, not blocking the calling thread (Asynchronous).
-	 * If asynchronous, ThreadModeDelegateForAsyncRunCompleted will define whether the callback delegate is called from the game thread (highly recommended) or from
-	 * any available thread (not fully thread safe).
-	 * - If ThreadModeDelegateForAsyncRunCompleted == ENeuralSynchronousMode::Asynchronous, the FOnAsyncRunCompleted delegate could be triggered from any thread.
-	 * - If ThreadModeDelegateForAsyncRunCompleted == ENeuralSynchronousMode::Synchronous, UNeuralNetwork::Run() will block the calling thread until completed, so a
-	 *   callback delegate is not required.
-	 * @see ENeuralSynchronousMode, ENeuralThreadMode for more details.
+	 * It defines whether UNeuralNetwork::Run() will block the thread until completed (Synchronous), or whether it will run on a background thread,
+	 * not blocking the calling thread (Asynchronous).
 	 */
 	UPROPERTY(Transient, EditAnywhere, Category = "Neural Network Inference")
 	ENeuralSynchronousMode SynchronousMode;
+
+	/**
+	 * If SynchronousMode is Asynchronous, this variable will define whether the callback delegate is called from the game thread (highly
+	 * recommended) or from any available thread (not fully thread safe).
+	 * - If this variable is set to ENeuralThreadMode::GameThread, the FOnAsyncRunCompleted delegate will be triggered from the main thread only.
+	 * - If this variable is set to ENeuralThreadMode::AnyThread, the FOnAsyncRunCompleted delegate could be triggered from any thread.
+	 * @see SynchronousMode for more details.
+	 */
 	UPROPERTY(Transient, EditAnywhere, Category = "Neural Network Inference")
 	ENeuralThreadMode ThreadModeDelegateForAsyncRunCompleted;
 
@@ -287,19 +298,28 @@ protected:
 	FString ModelFullFilePath;
 
 private:
+	/**
+	 * Whether this UNeuralNetwork instance has loaded a valid network model already, i.e., whether Load() was called and returned true.
+	 */
 	UPROPERTY(Transient, VisibleAnywhere, Category = "Neural Network Inference")
 	bool bIsLoaded;
 
+	/**
+	 * A buffer of memory representing the ONNX model. It is equivalent to the serialization of the ONNX file.
+	 * During Load(InModelFilePath), the ONNX file is read and stored as this buffer of raw bytes (TArray<uint8>).
+	 */
 	UPROPERTY()
 	TArray<uint8> ModelReadFromFileInBytes;
 
-	/** Whether some of the FNeuralTensor of InputTensor have flexible/variable dimensions. */
+	/**
+	 * Whether some of the FNeuralTensor of InputTensor have flexible/variable dimensions. E.g., this is useful for variable batch size.
+	 */
 	UPROPERTY(Transient, VisibleAnywhere, Category = "Neural Network Inference")
 	TArray<bool> AreInputTensorSizesVariable;
 
 	/**
 	 * Mutex to avoid issues or crashes due to the asynchronous Run() being run at the same time than any other non-const function.
-	 * @see UNeuralNetwork::Run().
+	 * @see UNeuralNetwork::Run() implementation for more details.
 	 */
 	FCriticalSection ResoucesCriticalSection;
 
@@ -309,35 +329,30 @@ private:
 	FOnAsyncRunCompleted OnAsyncRunCompletedDelegate;
 
 	/**
-	 * Stats-related members.
+	 * Statistics-related members used for GetLastRunTimeMSec(), GetRunStatistics(), GetInputMemoryTransferStats(), ResetStats().
 	 */
-	FNeuralStatisticsEstimator ComputeStatisticsEstimator;
+	FNeuralStatisticsEstimator RunStatisticsEstimator;
 	FNeuralStatisticsEstimator InputTransferStatisticsEstimator;
 
 	/**
-	 * Struct pointer containing the UE-and-ORT-based back end implementation.
-	 * PIMPL idiom to minimize memory when not using this back end and to hide 3rd party dependencies.
+	 * Struct pointers containing the UE-and-ORT-based (ImplBackEndUEAndORT) and only-UE-based (ImplBackEndUEOnly) back end implementations.
+	 * PIMPL idiom to minimize memory when not using each back end (only 1 is used at the time) and to hide 3rd party dependencies.
 	 * http://www.cppsamples.com/common-tasks/pimpl.html
 	 */
 	struct FImplBackEndUEAndORT;
 	TSharedPtr<FImplBackEndUEAndORT> ImplBackEndUEAndORT;
-
-	/**
-	 * Struct pointer containing the only-UE-based back end implementation.
-	 * PIMPL idiom to minimize memory when not using this back end.
-	 * http://www.cppsamples.com/common-tasks/pimpl.html
-	 */
 	struct FImplBackEndUEOnly;
 	TSharedPtr<FImplBackEndUEOnly> ImplBackEndUEOnly;
 
 	/**
-	 * It loads the desired network graph definition and weights internally saved on this UNeuralNetwork instance.
+	 * It loads the existing network graph definition and weights from the members of this UNeuralNetwork instance.
 	 * @return Whether the network was successfully loaded.
 	 */
 	bool Load();
 
 	/**
-	 * GetInputTensorMutable() and GetOutputTensorMutable() are the private and mutable version of GetInputTensor() and GetOutputTensor(), respectively.
+	 * GetInputTensorMutable() and GetOutputTensorMutable() are the private and mutable version of GetInputTensor() and GetOutputTensor(),
+	 * respectively.
 	 * Most FNeuralTensor properties are considered constant and fixed after calling UNeuralNetwork::Load(), other than the actual data values
 	 * (i.e., UnderlyingUInt8ArrayData) and some GPU properties. Thus, use these 2 functions with extreme precaution and be sure not to modify any
 	 * other FNeuralNetwork properties.
@@ -346,30 +361,30 @@ private:
 	FNeuralTensor& GetInputTensorMutable(const int32 InTensorIndex = 0);
 	FNeuralTensor& GetOutputTensorMutable(const int32 InTensorIndex = 0);
 
-public:
+public: // It should really say "private-for-user, public-for-auxiliary-NNI-classes"
 	/**
 	 * Internal function not needed by the user.
 	 * Used to create custom networks without an ONNX file for QA testing in FOperatorTester::TestOperator().
 	 */
 	bool Load(TArray<FNeuralTensor>& InTensors, const TArray<FNeuralTensor*>& InInputTensors, const TArray<FNeuralTensor*>& InOutputTensors,
 		const TArray<TSharedPtr<class FNeuralOperator>>& InOperators);
-	
+
 	/**
 	 * Internal enum class that should not be used by the user.
-	 * Whether UNeuralNetwork will use the highly optimized UnrealEngine-and-ONNXRuntime-based back end (UEAndORT) or the less optimized but fully
-	 * cross platform only-UE one (UEOnly).
-	 * We recommend using Auto, which will find and use the optimal back end for each platform.
+	 * It enumerates the different types of back ends. Use Auto, which will find and use the optimal back end for each platform.
 	 */
 	enum class ENeuralBackEnd : uint8
 	{
 		/**
-		 * UnrealEngine-and-ONNXRuntime-accelerated back end, ideal for those platforms that support it. WITH_UE_AND_ORT_SUPPORT is the C++ define
-		 * macro that checks whether UEAndORT support exists for the current platform.
+		 * Highly optimized UnrealEngine-and-ONNXRuntime-based back end, ideal for those platforms that support it. WITH_UE_AND_ORT_SUPPORT is the
+		 * C++ define macro that checks whether UEAndORT support exists for the current platform.
 		 */
 		UEAndORT,
-		/** It might be slower than the UEAndORT back end, but it will compile in all platforms and OSs compatible with Unreal Engine. */
+		/**
+		 * Slower than the UEAndORT back end maintained for now for historical reasons. There is no benefit in using this back end.
+		 */
 		UEOnly,
-		/** Recommended value. It will use the efficient UEAndORT if supported by the platform, and fall back to UEOnly otherwise. */
+		/** Recommended value. It will use the efficient UEAndORT if supported by the platform, and try to fall back to UEOnly otherwise. */
 		Auto
 	};
 
@@ -384,14 +399,17 @@ public:
 	 * SetBackEnd() is NOT THREAD SAFE, make sure that other non-const functions such as Run() are not running when SetBackEnd is called.
 	 * @see ENeuralBackEnd for more details.
 	 */
+	UE_DEPRECATED(5.0, "Do not use, ENeuralBackEnd will be removed in future versions and only UEAndORT back end will be supported.")
 	ENeuralBackEnd GetBackEnd() const;
+	UE_DEPRECATED(5.0, "Do not use, ENeuralBackEnd will be removed in future versions and only UEAndORT back end will be supported.")
 	ENeuralBackEnd GetBackEndForCurrentPlatform() const;
+	UE_DEPRECATED(5.0, "Do not use, ENeuralBackEnd will be removed in future versions and only UEAndORT back end will be supported.")
 	bool SetBackEnd(const ENeuralBackEnd InBackEnd);
 
 #if WITH_EDITOR
 	/**
 	 * Internal and Editor-only functions not needed by the user.
-	 * Importing data and options used for loading the network.
+	 * They provide importing data and options used for loading the network.
 	 */
 	TObjectPtr<class UAssetImportData> GetAssetImportData() const;
 	TObjectPtr<class UAssetImportData> GetAndMaybeCreateAssetImportData();
@@ -399,13 +417,12 @@ public:
 
 private:
 	/**
-	 * Internal variable that should not be used by the user.
-	 * @see ENeuralBackEnd for more details.
+	 * The back end to be used (it could and should be set to Auto). @see ENeuralBackEnd for more details.
 	 */
 	ENeuralBackEnd BackEnd;
 
 	/**
-	 * Internal variable that should not be used by the user.
+	 * The actual back end currently being used (it cannot be set to Auto).
 	 * If BackEnd != Auto, BackEndForCurrentPlatform will be equal to BackEnd.
 	 * Otherwise, BackEndForCurrentPlatform will be set to the optimal BackEnd given the current platform.
 	 * @see ENeuralBackEnd for more details.
