Beter commented code.

#rnx

diff --git a/Engine/Plugins/Experimental/NeuralNetworkInference/Source/NeuralNetworkInference/Public/NeuralNetwork.h b/Engine/Plugins/Experimental/NeuralNetworkInference/Source/NeuralNetworkInference/Public/NeuralNetwork.h
--- a/Engine/Plugins/Experimental/NeuralNetworkInference/Source/NeuralNetworkInference/Public/NeuralNetwork.h
+++ b/Engine/Plugins/Experimental/NeuralNetworkInference/Source/NeuralNetworkInference/Public/NeuralNetwork.h
@@ -14,7 +14,7 @@
  * - Ease-of-use: Simple but powerful API.
  * - Completeness: All the functionality of any state-of-the-art deep learning framework.
  *
- * UNeuralNetwork is the key class of NNI, and the only one users should interact with. It represents the deep neural model itself. It is capable of
+ * UNeuralNetwork is the key class of NNI, and the main one users should interact with. It represents the deep neural model itself. It is capable of
  * loading and running inference (i.e., a forward pass) on any ONNX (Open Neural Network eXchange) model. ONNX is the industry standard for ML
  * interoperability, and all major frameworks (PyTorch, TensorFlow, MXNet, Caffe2, etc.) provide converters to ONNX.
  *
@@ -41,7 +41,7 @@
  *			UE_LOG(LogTemp, Warning, TEXT("UNeuralNetwork could not loaded from %s."), *ONNXModelFilePath);
  *		}
  *
- * 1b Loading a UNeuralNetwork from a previously-created UNeuralNetwork UAsset
+ * 1b Loading an existing UNeuralNetwork from its UAsset
  *		// Load the UNeuralNetwork object from a UNeuralNetwork UAsset
  *		const FString NetworkUAssetFilePath = TEXT("'/Game/Models/ExampleNetwork/ExampleNetwork.ExampleNetwork'");
  *		UNeuralNetwork* Network = LoadObject<UNeuralNetwork>((UObject*)GetTransientPackage(), *NetworkUAssetFilePath);
diff --git a/Engine/Plugins/Experimental/NeuralNetworkInference/Source/Private/NeuralNetworkInferenceQA/Private/NeuralNetworkInferenceQA.cpp b/Engine/Plugins/Experimental/NeuralNetworkInference/Source/Private/NeuralNetworkInferenceQA/Private/NeuralNetworkInferenceQA.cpp
--- a/Engine/Plugins/Experimental/NeuralNetworkInference/Source/Private/NeuralNetworkInferenceQA/Private/NeuralNetworkInferenceQA.cpp
+++ b/Engine/Plugins/Experimental/NeuralNetworkInference/Source/Private/NeuralNetworkInferenceQA/Private/NeuralNetworkInferenceQA.cpp
@@ -34,15 +34,10 @@ IMPLEMENT_SIMPLE_AUTOMATION_TEST(FNeuralNetworkInferenceTest, "System.Engine.Mac
 
 bool FNeuralNetworkInferenceTest::RunTest(const FString& Parameters)
 {
-// #if WITH_EDITOR
 	if (!UNeuralNetworkInferenceQA::UnitTesting())
 	{
 		AddError(TEXT("FNeuralNetworkInferenceTest::RunTest(): UNeuralNetworkInferenceQA::UnitTesting() returned false. See log above.")); // Both AddWarning/AddError would work
 	}
-// #else //WITH_EDITOR
-// 	UE_LOG(LogNeuralNetworkInferenceQA, Warning, TEXT("@todo: FNeuralNetworkInferenceTest::RunTest() temporarily disabled on game due to JIRA"
-// 		" jira.it.epicgames.com/browse/UE-134683, UAssets not found on the Horde cooked tests."));
-// #endif //WITH_EDITOR
 	return true;
 }
 
