NNI: Rebind buffer to DirectML on each run of inference.
This fixes issue where RHI buffer defrag would cause stale binding.
Change adds support for inference context which has the bonus that it allows instancing.
New API is added for this, and old API, along with default tensors remains.
New API only supports limited use cases (ONNX with GPU inference, with CPU input and GPU output).
That is enough for MLDeformer.
New API will log error and not do inference if it doesn't support the use case.
If new API is kept, then we can extend to other use cases and possibly deprecate old API.
(Other design changes with NNI might mean this new API is short lived).

diff --git a/Engine/Plugins/Experimental/NeuralNetworkInference/Source/NeuralNetworkInference/Private/NeuralNetwork.cpp b/Engine/Plugins/Experimental/NeuralNetworkInference/Source/NeuralNetworkInference/Private/NeuralNetwork.cpp
--- a/Engine/Plugins/Experimental/NeuralNetworkInference/Source/NeuralNetworkInference/Private/NeuralNetwork.cpp
+++ b/Engine/Plugins/Experimental/NeuralNetworkInference/Source/NeuralNetworkInference/Private/NeuralNetwork.cpp
@@ -705,3 +705,57 @@ void UNeuralNetwork::Serialize(FArchive& Archive)
 #endif // WITH_EDITORONLY_DATA
 	Super::Serialize(Archive);
 }
+
+
+int32 UNeuralNetwork::CreateInferenceContext()
+{
+	if (!bIsLoaded)
+	{
+		UE_LOG(LogNeuralNetworkInference, Warning, TEXT("UNeuralNetwork::CreateInferenceContext(): Call UNeuralNetwork::Load() to load a model first."));
+		return -1;
+	}
+	if (BackEndForCurrentPlatform != ENeuralBackEnd::UEAndORT || SynchronousMode != ENeuralSynchronousMode::Synchronous)
+	{
+		UE_LOG(LogNeuralNetworkInference, Warning, TEXT("UNeuralNetwork::CreateInferenceContext(): Execution contexts currently only support synchronous Onyx models."));
+		return -1;
+	}
+	if (DeviceType != ENeuralDeviceType::GPU || InputDeviceType != ENeuralDeviceType::CPU || OutputDeviceType != ENeuralDeviceType::GPU)
+	{
+		UE_LOG(LogNeuralNetworkInference, Warning, TEXT("UNeuralNetwork::CreateInferenceContext(): Execution contexts currently only support GPU device with CPU input -> GPU ouput."));
+		return -1;
+	}
+
+	return ImplBackEndUEAndORT->CreateInferenceContext();
+}
+
+void UNeuralNetwork::DestroyInferenceContext(int32 ContextHandle)
+{
+	return ImplBackEndUEAndORT->DestroyInferenceContext(ContextHandle);
+}
+
+void UNeuralNetwork::Run(FRDGBuilder& GraphBuilder, int32 ContextHandle)
+{
+	DECLARE_SCOPE_CYCLE_COUNTER(TEXT("UNeuralNetwork_Run"), STAT_UNeuralNetwork_Run, STATGROUP_MachineLearning);
+
+	FNeuralTimer RunTimer;
+	RunTimer.Tic();
+
+	ImplBackEndUEAndORT->Run(GraphBuilder, ContextHandle);
+
+	RunStatisticsEstimator.StoreSample(RunTimer.Toc());
+}
+
+const FNeuralTensor& UNeuralNetwork::GetInputTensorForContext(const int32 InContextHandle, const int32 InTensorIndex) const
+{
+	return ImplBackEndUEAndORT->GetInputTensorForContext(InContextHandle, InTensorIndex);
+}
+
+const FNeuralTensor& UNeuralNetwork::GetOutputTensorForContext(const int32 InContextHandle, const int32 InTensorIndex) const
+{
+	return ImplBackEndUEAndORT->GetOutputTensorForContext(InContextHandle, InTensorIndex);
+}
+
+void* UNeuralNetwork::GetInputDataPointerMutableForContext(const int32 InContextHandle, const int32 InTensorIndex)
+{
+	return ImplBackEndUEAndORT->GetInputTensorForContext(InContextHandle, InTensorIndex).GetData();
+}
diff --git a/Engine/Plugins/Experimental/NeuralNetworkInference/Source/NeuralNetworkInference/Private/NeuralNetworkImplBackEndUEAndORT.cpp b/Engine/Plugins/Experimental/NeuralNetworkInference/Source/NeuralNetworkInference/Private/NeuralNetworkImplBackEndUEAndORT.cpp
--- a/Engine/Plugins/Experimental/NeuralNetworkInference/Source/NeuralNetworkInference/Private/NeuralNetworkImplBackEndUEAndORT.cpp
+++ b/Engine/Plugins/Experimental/NeuralNetworkInference/Source/NeuralNetworkInference/Private/NeuralNetworkImplBackEndUEAndORT.cpp
@@ -44,6 +44,7 @@ NNI_THIRD_PARTY_INCLUDES_END
 
 #include "ShaderParameterUtils.h"
 #include "ShaderParameterStruct.h"
+#include "RenderGraphUtils.h"
 
 #ifdef WITH_UE_AND_ORT_SUPPORT
 
@@ -806,12 +807,14 @@ bool UNeuralNetwork::FImplBackEndUEAndORT::SetTensorsFromNetwork(TArray<FNeuralT
 			}
 
 			// Link tensor with ORT blob
-			if (!LinkTensorResourceToONNXRuntime(OutTensors[TensorIndex], OrtTensors[TensorIndex], D3DResource))
+			void* DmlGPUAllocation = LinkTensorResourceToONNXRuntime(DmlApi, DmlGPUMemoryInfo.Get(), OutTensors[TensorIndex], OrtTensors[TensorIndex], D3DResource);
+			if  (DmlGPUAllocation == nullptr)
 			{
 				UE_LOG(LogNeuralNetworkInference, Warning,
 					TEXT("FImplBackEndUEAndORT::SetTensorsFromNetwork(): Failed to link the GPU resource to ONNX Runtime."));
 				return false;
 			}
+			DmlGPUResources.Emplace(DmlGPUAllocation);
 		}
 		else
 #endif
@@ -819,24 +822,22 @@ bool UNeuralNetwork::FImplBackEndUEAndORT::SetTensorsFromNetwork(TArray<FNeuralT
 			// Pre-allocate TArray (if size is different)
 			OutTensors[TensorIndex].SetNumUninitialized(InTensorDataTypes[TensorIndex], InSizes[TensorIndex]);
 			// Link tensor with ORT blob
-			LinkTensorToONNXRuntime(OutTensors, OrtTensors, *AllocatorInfo, TensorIndex);
+			LinkTensorToONNXRuntime(OutTensors[TensorIndex], OrtTensors[TensorIndex], *AllocatorInfo);
 		}
 	}
 
 	return true;
 }
 
-void UNeuralNetwork::FImplBackEndUEAndORT::LinkTensorToONNXRuntime(TArray<FNeuralTensor>& InOutTensors, TArray<Ort::Value>& InOutOrtTensors,
-	Ort::MemoryInfo& InOutAllocatorInfo, const int32 InTensorIndex)
+void UNeuralNetwork::FImplBackEndUEAndORT::LinkTensorToONNXRuntime(FNeuralTensor& InTensor, Ort::Value& OutOrtTensor, Ort::MemoryInfo& InOutAllocatorInfo)
 {
-	const TArray<int64>& Sizes = InOutTensors[InTensorIndex].GetSizes();
-	if (Sizes.Num() > 0 && InOutTensors[InTensorIndex].Num() > 0)
+	const TArray<int64>& Sizes = InTensor.GetSizes();
+	if (Sizes.Num() > 0 && InTensor.Num() > 0)
 	{
-		FNeuralTensor& Tensor = InOutTensors[InTensorIndex];
-		const int64 Volume = Tensor.Num();
+		const int64 Volume = InTensor.Num();
 		const int32 ArrayDimensions = Sizes.Num();
 
-		const ENeuralDataType NeuralDataType = Tensor.GetDataType();
+		const ENeuralDataType NeuralDataType = InTensor.GetDataType();
 		if (NeuralDataType == ENeuralDataType::Float)
 		{
 #ifdef _WIN32
@@ -847,40 +848,38 @@ void UNeuralNetwork::FImplBackEndUEAndORT::LinkTensorToONNXRuntime(TArray<FNeura
 			SizesInt64t.SetNumUninitialized(ArrayDimensions);
 			FMemory::Memcpy(SizesInt64t.GetData(), (int64_t*)Sizes.GetData(), sizeof(int64_t) * ArrayDimensions);
 #endif //_WIN32
-			InOutOrtTensors[InTensorIndex] = Ort::Value::CreateTensor<float>(InOutAllocatorInfo, Tensor.GetDataCasted<float>(), Volume, SizesInt64t.GetData(), ArrayDimensions);
+			OutOrtTensor = Ort::Value::CreateTensor<float>(InOutAllocatorInfo, InTensor.GetDataCasted<float>(), Volume, SizesInt64t.GetData(), ArrayDimensions);
 		}
 		//else if (NeuralDataType == ENeuralDataType::Double)
 		//{
-		//	InOutOrtTensors[InTensorIndex] = Ort::Value::CreateTensor<double>(InOutAllocatorInfo, Tensor.GetDataCasted<double>(), Volume, Sizes.GetData(), ArrayDimensions);
+		//	OutOrtTensor = Ort::Value::CreateTensor<double>(InOutAllocatorInfo, InTensor.GetDataCasted<double>(), Volume, Sizes.GetData(), ArrayDimensions);
 		//}
 		else
 		{
-			UE_LOG(LogNeuralNetworkInference, Warning, TEXT("FImplBackEndUEAndORT::LinkTensorToONNXRuntime(): Not implemented (yet) for ENeuralDataType = %d."), (int32)NeuralDataType);
+			UE_LOG(LogNeuralNetworkInference, Warning, TEXT("LinkTensorToONNXRuntime(): Not implemented (yet) for ENeuralDataType = %d."), (int32)NeuralDataType);
 		}
 	}
 }
 
 #ifdef PLATFORM_WIN64
-bool UNeuralNetwork::FImplBackEndUEAndORT::LinkTensorResourceToONNXRuntime(FNeuralTensor& InOutTensor, Ort::Value& InOutOrtTensor, void* InOutD3DResource)
+void* UNeuralNetwork::FImplBackEndUEAndORT::LinkTensorResourceToONNXRuntime(OrtDmlApi const* InDmlApi, Ort::MemoryInfo* InMemoryInfo, FNeuralTensor& InOutTensor, Ort::Value& InOutOrtTensor, void* InD3DResource)
 {
-	if (!DmlApi)
+	if (!InDmlApi)
 	{
-		UE_LOG(LogNeuralNetworkInference, Warning, TEXT("FImplBackEndUEAndORT::LinkTensorResourceToONNXRuntime(): DmlGPUAllocator is not valid."));
-		return false;
+		UE_LOG(LogNeuralNetworkInference, Warning, TEXT("LinkTensorResourceToONNXRuntime(): DmlGPUAllocator is not valid."));
+		return nullptr;
 	}
 
 	void* DmlGPUAllocation = nullptr;
 	
-	DmlApi->CreateGPUAllocationFromD3DResource(reinterpret_cast<ID3D12Resource*>(InOutD3DResource), &DmlGPUAllocation);
+	InDmlApi->CreateGPUAllocationFromD3DResource(reinterpret_cast<ID3D12Resource*>(InD3DResource), &DmlGPUAllocation);
 
 	if (!DmlGPUAllocation)
 	{
-		UE_LOG(LogNeuralNetworkInference, Warning, TEXT("FImplBackEndUEAndORT::LinkTensorResourceToONNXRuntime(): DmlGPUAllocation is nullptr."));
-		return false;
+		UE_LOG(LogNeuralNetworkInference, Warning, TEXT("LinkTensorResourceToONNXRuntime(): DmlGPUAllocation is nullptr."));
+		return nullptr;
 	}
 
-	DmlGPUResources.Emplace(DmlGPUAllocation);
-
 	const TArray<int64>& Sizes = InOutTensor.GetSizes();
 	if (Sizes.Num() > 0 && InOutTensor.Num() > 0)
 	{
@@ -897,7 +896,7 @@ bool UNeuralNetwork::FImplBackEndUEAndORT::LinkTensorResourceToONNXRuntime(FNeur
 			SizesInt64t.SetNumUninitialized(ArrayDimensions);
 			FMemory::Memcpy(SizesInt64t.GetData(), (int64_t*)Sizes.GetData(), sizeof(int64_t) * ArrayDimensions);
 #endif //_WIN32
-			InOutOrtTensor = Ort::Value::CreateTensor(*DmlGPUMemoryInfo.Get(), DmlGPUAllocation, InOutTensor.NumInBytes(), SizesInt64t.GetData(),
+			InOutOrtTensor = Ort::Value::CreateTensor(*InMemoryInfo, DmlGPUAllocation, InOutTensor.NumInBytes(), SizesInt64t.GetData(),
 				ArrayDimensions, ONNX_TENSOR_ELEMENT_DATA_TYPE_FLOAT);
 		}
 		//else if (NeuralDataType == ENeuralDataType::Double)
@@ -909,11 +908,12 @@ bool UNeuralNetwork::FImplBackEndUEAndORT::LinkTensorResourceToONNXRuntime(FNeur
 		{
 			UE_LOG(LogNeuralNetworkInference, Warning,
 				TEXT("FImplBackEndUEAndORT::LinkTensorToONNXRuntime(): Not implemented (yet) for ENeuralDataType = %d."), (int32)NeuralDataType);
-			return false;
+			InDmlApi->FreeGPUAllocation(DmlGPUAllocation);
+			return nullptr;
 		}
 	}
 
-	return true;
+	return DmlGPUAllocation;
 }
 #endif
 
@@ -1055,4 +1055,198 @@ void UNeuralNetwork::FImplBackEndUEAndORT::ClearResources()
 #endif //PLATFORM_WIN64
 }
 
+bool UNeuralNetwork::FImplBackEndUEAndORT::FInferenceContext::Init(Ort::Session* InSession, Ort::AllocatorWithDefaultOptions* InAllocator, Ort::MemoryInfo* InAllocatorInfo)
+{
+	// Current assumptions here (all these things can be fixed with more work): 
+	// Input tensors are CPU
+	// Output tensors are GPU 
+	// Float data types only
+	// No variable dimensions
+
+	ENeuralTensorType TensorTypes[] = { ENeuralTensorType::Input, ENeuralTensorType::Output };
+
+	for (ENeuralTensorType TensorType : TensorTypes)
+	{
+		const bool bInputTensors = TensorType == ENeuralTensorType::Input;
+		const uint32 NumTensors = bInputTensors ? InSession->GetInputCount() : InSession->GetOutputCount();
+
+		for (uint32 TensorIndex = 0; TensorIndex < NumTensors; ++TensorIndex)
+		{
+			const char* TensorName = bInputTensors ? InSession->GetInputName(TensorIndex, *InAllocator) : InSession->GetOutputName(TensorIndex, *InAllocator);
+			Ort::TypeInfo TypeInfo = bInputTensors ? InSession->GetInputTypeInfo(TensorIndex) : InSession->GetOutputTypeInfo(TensorIndex);
+			const Ort::TensorTypeAndShapeInfo TensorInfo = TypeInfo.GetTensorTypeAndShapeInfo();
+
+			ENeuralDataType TensorDataType = ENeuralDataType::None;
+			{
+				const ONNXTensorElementDataType ONNXDataType = TensorInfo.GetElementType();
+				if (ONNXDataType == ONNX_TENSOR_ELEMENT_DATA_TYPE_FLOAT)
+				{
+					TensorDataType = ENeuralDataType::Float;
+				}
+				else
+				{
+					UE_LOG(LogNeuralNetworkInference, Warning,
+						TEXT("FInferenceContext::Init(): ONNXDataType = %d not implemented yet."), (int32)ONNXDataType);
+				}
+			}
+
+			TArray<int64> TensorSizes;
+			for (int64 TensorSize : TensorInfo.GetShape())
+			{
+				if (TensorSize < 0)
+				{
+					TensorSize = 1;
+					UE_LOG(LogNeuralNetworkInference, Warning,
+						TEXT("FInferenceContext::Init(): Negative (i.e., variable) dimensions not allowed yet, hard-coded to 1."
+							" Let us know if you really need variable dimensions."
+							" Keep in mind that fixed sizes might allow additional optimizations and speedup of the network during Run()."));
+				}
+
+				TensorSizes.Add(TensorSize);
+			}
+
+			if (bInputTensors)
+			{
+				InputTensorNames.Emplace(TensorName);
+
+				InputTensors.Emplace(FNeuralTensor(TensorDataType, TensorSizes, ANSI_TO_TCHAR(TensorName), TensorType));
+
+				InputOrtTensors.Emplace(Ort::Value(nullptr));
+				LinkTensorToONNXRuntime(InputTensors[TensorIndex], InputOrtTensors[TensorIndex], *InAllocatorInfo);
+			}
+			else
+			{
+				OutputTensorNames.Emplace(TensorName);
+
+				FNeuralTensor& Tensor = OutputTensors.Emplace_GetRef(FNeuralTensor(TensorDataType, TensorSizes, ANSI_TO_TCHAR(TensorName), TensorType));
+				Tensor.SetEnableGPU(true);
+
+				OutputOrtTensors.Emplace(Ort::Value(nullptr));
+			}
+
+			TypeInfo.release();
+		}
+	}
+
+	return true;
+}
+
+void UNeuralNetwork::FImplBackEndUEAndORT::FInferenceContext::UpdateGPUAllocations(FRDGBuilder& GraphBuilder)
+{
+	for (int32 TensorIndex = 0; TensorIndex < OutputOrtTensors.Num(); ++TensorIndex)
+	{
+		OutputTensors[TensorIndex].InitPooledBufferForUEAndORTBackEnd_RenderThread(GraphBuilder);
+	}
+}
+
+void UNeuralNetwork::FImplBackEndUEAndORT::FInferenceContext::Release()
+{
+	InputTensorNames.Reset();
+	InputTensors.Reset();
+	InputOrtTensors.Reset();
+	OutputTensorNames.Reset();
+	OutputTensors.Reset();
+	OutputOrtTensors.Reset();
+}
+
+#ifdef PLATFORM_WIN64
+
+void UNeuralNetwork::FImplBackEndUEAndORT::FInferenceContext::BindDMLAllocations(OrtDmlApi const* InDmlApi, Ort::MemoryInfo* InMemoryInfo)
+{
+	ReleaseDMLAllocations(InDmlApi);
+	
+	OutputDmlAllocation.SetNum(OutputOrtTensors.Num());
+	for (int32 TensorIndex = 0; TensorIndex < OutputOrtTensors.Num(); ++TensorIndex)
+	{
+		void* D3DResource = FNeuralNetworkInferenceUtilsGPU::GetD3D12Resource(OutputTensors[TensorIndex].GetPooledBuffer()->GetRHI());
+		OutputDmlAllocation[TensorIndex] = LinkTensorResourceToONNXRuntime(InDmlApi, InMemoryInfo, OutputTensors[TensorIndex], OutputOrtTensors[TensorIndex], D3DResource);
+	}
+}
+
+void UNeuralNetwork::FImplBackEndUEAndORT::FInferenceContext::ReleaseDMLAllocations(OrtDmlApi const* InDmlApi)
+{
+	for (int32 TensorIndex = 0; TensorIndex < OutputDmlAllocation.Num(); ++TensorIndex)
+	{
+		InDmlApi->FreeGPUAllocation(OutputDmlAllocation[TensorIndex]);
+	}
+	OutputDmlAllocation.Reset();
+}
+
+#endif // PLATFORM_WIN64
+
 #endif //WITH_UE_AND_ORT_SUPPORT
+
+int32 UNeuralNetwork::FImplBackEndUEAndORT::CreateInferenceContext()
+{
+#ifdef WITH_UE_AND_ORT_SUPPORT
+	const int32 Handle = Contexts.Emplace();
+	Contexts[Handle].Init(Session.Get(), Allocator.Get(), AllocatorInfo.Get());
+	return Handle;
+#else
+	UE_LOG(LogNeuralNetworkInference, Warning, TEXT("FImplBackEndUEAndORT::CreateInferenceContext(): Platform or Operating System not suported yet for UEAndORT"
+		" BackEnd. Set BackEnd to ENeuralBackEnd::Auto or ENeuralBackEnd::UEOnly for this platform."));
+	return -1;
+#endif
+}
+
+void UNeuralNetwork::FImplBackEndUEAndORT::DestroyInferenceContext(int32 InContextHandle)
+{
+#ifdef WITH_UE_AND_ORT_SUPPORT
+	check(Contexts.IsValidIndex(InContextHandle));
+#ifdef PLATFORM_WIN64
+	Contexts[InContextHandle].ReleaseDMLAllocations(DmlApi);
+#endif
+	Contexts[InContextHandle].Release();
+	Contexts.RemoveAt(InContextHandle);
+#endif
+}
+
+void UNeuralNetwork::FImplBackEndUEAndORT::Run(FRDGBuilder& GraphBuilder, int32 InContextHandle)
+{
+#ifdef WITH_UE_AND_ORT_SUPPORT
+#if WITH_EDITOR
+	try
+#endif //WITH_EDITOR
+	{
+		Contexts[InContextHandle].UpdateGPUAllocations(GraphBuilder);
+
+#ifdef PLATFORM_WIN64
+		Contexts[InContextHandle].BindDMLAllocations(DmlApi, DmlGPUMemoryInfo.Get());
+#endif
+
+		AddPass(GraphBuilder, RDG_EVENT_NAME("OrtRun"), [InSession = Session.Get(), &Context = Contexts[InContextHandle]](FRHICommandListImmediate&)
+		{
+			InSession->Run(
+				Ort::RunOptions{ nullptr },
+				Context.InputTensorNames.GetData(), & Context.InputOrtTensors[0], Context.InputTensorNames.Num(),
+				Context.OutputTensorNames.GetData(), & Context.OutputOrtTensors[0], Context.OutputTensorNames.Num());
+		});
+	}
+#if WITH_EDITOR
+	catch (const std::exception& Exception)
+	{
+		UE_LOG(LogNeuralNetworkInference, Error, TEXT("%s"), UTF8_TO_TCHAR(Exception.what()));
+	}
+#endif //WITH_EDITOR
+#endif //WITH_UE_AND_ORT_SUPPORT
+}
+
+FNeuralTensor& UNeuralNetwork::FImplBackEndUEAndORT::GetInputTensorForContext(int32 InContextHandle, int32 InTensorIndex)
+{
+#ifdef WITH_UE_AND_ORT_SUPPORT
+	check(Contexts.IsValidIndex(InContextHandle));
+	return Contexts[InContextHandle].InputTensors[InTensorIndex];
+#else
+	return InputTensors[InTensorIndex];
+#endif
+}
+
+FNeuralTensor& UNeuralNetwork::FImplBackEndUEAndORT::GetOutputTensorForContext(int32 InContextHandle, int32 InTensorIndex)
+{
+#ifdef WITH_UE_AND_ORT_SUPPORT
+	check(Contexts.IsValidIndex(InContextHandle));
+	return Contexts[InContextHandle].OutputTensors[InTensorIndex];
+#else
+	return OutputTensors[InTensorIndex];
+#endif
+}
diff --git a/Engine/Plugins/Experimental/NeuralNetworkInference/Source/NeuralNetworkInference/Private/NeuralNetworkImplBackEndUEAndORT.h b/Engine/Plugins/Experimental/NeuralNetworkInference/Source/NeuralNetworkInference/Private/NeuralNetworkImplBackEndUEAndORT.h
--- a/Engine/Plugins/Experimental/NeuralNetworkInference/Source/NeuralNetworkInference/Private/NeuralNetworkImplBackEndUEAndORT.h
+++ b/Engine/Plugins/Experimental/NeuralNetworkInference/Source/NeuralNetworkInference/Private/NeuralNetworkImplBackEndUEAndORT.h
@@ -56,8 +56,15 @@ public:
 		const TArray<uint8>& InModelReadFromFileInBytes, const FString& InModelFullFilePath, const ENeuralDeviceType InDeviceType, const ENeuralDeviceType InInputDeviceType,
 		const ENeuralDeviceType InOutputDeviceType);
 
+	int32 CreateInferenceContext();
+	void DestroyInferenceContext(int32 ContextHandle);
+
+	void Run(FRDGBuilder& GraphBuilder, int32 ContextHandle);
 	void Run(const ENeuralSynchronousMode InSynchronousMode, const ENeuralDeviceType InDeviceType, const ENeuralDeviceType InInputDeviceType);
 
+	FNeuralTensor& GetInputTensorForContext(int32 InContextHandle, int32 InTensorIndex);
+	FNeuralTensor& GetOutputTensorForContext(int32 InContextHandle, int32 InTensorIndex);
+
 #ifdef WITH_UE_AND_ORT_SUPPORT
 private:
 	/** Async support */
@@ -86,7 +93,35 @@ private:
 	TArray<Ort::Value> OutputOrtTensors;
 	/** Tensor names */
 	TArray<const char*> OutputTensorNames;
-	
+
+	/**
+	 * Storage for input and output tensors for ONNX backend.
+	 * GPU memory is allocated by RDG/RHI and bound to Ort handles.
+	 * Because RHI can defrag the RDG pool at any time we need to call UpdateGPUAllocations() for each frame we use.
+	 */
+	struct FInferenceContext
+	{
+		TArray<const char*> InputTensorNames;
+		TArray<FNeuralTensor> InputTensors;
+		TArray<Ort::Value> InputOrtTensors;
+
+		TArray<const char*> OutputTensorNames;
+		TArray<FNeuralTensor> OutputTensors;
+		TArray<Ort::Value> OutputOrtTensors;
+		TArray<void*> OutputDmlAllocation;
+
+		bool Init(Ort::Session* Session, Ort::AllocatorWithDefaultOptions* Allocator, Ort::MemoryInfo* AllocatorInfo);
+		void UpdateGPUAllocations(FRDGBuilder& GraphBuilder);
+		void Release();
+#ifdef PLATFORM_WIN64
+		void BindDMLAllocations(OrtDmlApi const* InDmlApi, Ort::MemoryInfo* InMemoryInfo);
+		void ReleaseDMLAllocations(OrtDmlApi const* InDmlApi);
+#endif
+
+	};
+	/** Array of allocated inference contexts. */
+	TSparseArray<FInferenceContext> Contexts;
+
 	// Helper class to run session as an async task
 	class FNeuralNetworkAsyncTask : public FNonAbandonableTask
 	{
@@ -125,11 +160,9 @@ private:
 		TArray<TArray<int64>>& InSizes, TArray<ENeuralTensorType>& InTensorTypes, const bool bIsInput,
 		const ENeuralDeviceType InInputDeviceType, const ENeuralDeviceType InOutputDeviceType);
 
-	static void LinkTensorToONNXRuntime(TArray<FNeuralTensor>& InOutTensors, TArray<Ort::Value>& InOutOrtTensors, Ort::MemoryInfo& InOutAllocatorInfo,
-		const int32 InTensorIndex);
-
+	static void LinkTensorToONNXRuntime(FNeuralTensor& InTensor, Ort::Value& OutOrtTensor, Ort::MemoryInfo& InOutAllocatorInfo);
 #ifdef PLATFORM_WIN64
-	bool LinkTensorResourceToONNXRuntime(FNeuralTensor& InOutTensor, Ort::Value& InOutOrtTensor, void* InOutD3DResource);
+	static void* LinkTensorResourceToONNXRuntime(OrtDmlApi const* InDmlApi, Ort::MemoryInfo* InMemoryInfo, FNeuralTensor& InOutTensor, Ort::Value& InOutOrtTensor, void* InD3DResource);
 #endif
 
 	void RunSessionAsync(const ENeuralDeviceType InDeviceType, const ENeuralDeviceType InInputDeviceType);
diff --git a/Engine/Plugins/Experimental/NeuralNetworkInference/Source/NeuralNetworkInference/Private/NeuralTensor.cpp b/Engine/Plugins/Experimental/NeuralNetworkInference/Source/NeuralNetworkInference/Private/NeuralTensor.cpp
--- a/Engine/Plugins/Experimental/NeuralNetworkInference/Source/NeuralNetworkInference/Private/NeuralTensor.cpp
+++ b/Engine/Plugins/Experimental/NeuralNetworkInference/Source/NeuralNetworkInference/Private/NeuralTensor.cpp
@@ -637,18 +637,7 @@ bool FNeuralTensor::InitPooledBufferForUEAndORTBackEnd(void** InOutNativeResourc
 		[this, &InOutNativeResource, &bIsGraphBuilderDone] (FRHICommandListImmediate& RHICmdList)
 		{
 			FRDGBuilder Builder(RHICmdList);
-
-			FRDGBufferDesc BufferDesc;
-			BufferDesc.BytesPerElement = FNeuralDataTypeUtils::GetByteSize(DataType);
-			BufferDesc.NumElements = Num();
-			BufferDesc.Usage = EBufferUsageFlags::UnorderedAccess | EBufferUsageFlags::ShaderResource;
-			BufferDesc.UnderlyingType = FRDGBufferDesc::EUnderlyingType::VertexBuffer;
-
-			FRDGBufferRef BufferRef = Builder.CreateBuffer(BufferDesc, *GetName());
-
-			// Recreate PooledBuffer for future runs
-			*PooledBuffer = Builder.ConvertToExternalBuffer(BufferRef);
-
+			InitPooledBufferForUEAndORTBackEnd_RenderThread(Builder);
 			Builder.Execute();
 
 #ifdef PLATFORM_WIN64
@@ -672,6 +661,23 @@ bool FNeuralTensor::InitPooledBufferForUEAndORTBackEnd(void** InOutNativeResourc
 	return true;
 }
 
+void FNeuralTensor::InitPooledBufferForUEAndORTBackEnd_RenderThread(FRDGBuilder& GraphBuilder)
+{
+	FRDGBufferDesc BufferDesc;
+	BufferDesc.BytesPerElement = FNeuralDataTypeUtils::GetByteSize(DataType);
+	BufferDesc.NumElements = Num();
+	BufferDesc.Usage = EBufferUsageFlags::UnorderedAccess | EBufferUsageFlags::ShaderResource;
+	BufferDesc.UnderlyingType = FRDGBufferDesc::EUnderlyingType::VertexBuffer;
+	FRDGBufferRef BufferRef = GraphBuilder.CreateBuffer(BufferDesc, *GetName());
+
+	// Extract buffer immediately. Need this so that we can potentially bind to DirectML.
+	if (!PooledBuffer.IsValid())
+	{
+		PooledBuffer = MakeShared<TRefCountPtr<FRDGPooledBuffer>>();
+	}
+	*PooledBuffer = GraphBuilder.ConvertToExternalBuffer(BufferRef);
+}
+
 const TRefCountPtr<FRDGPooledBuffer>& FNeuralTensor::GetPooledBuffer() const
 {
 	// Sanity checks
diff --git a/Engine/Plugins/Experimental/NeuralNetworkInference/Source/NeuralNetworkInference/Public/NeuralNetwork.h b/Engine/Plugins/Experimental/NeuralNetworkInference/Source/NeuralNetworkInference/Public/NeuralNetwork.h
--- a/Engine/Plugins/Experimental/NeuralNetworkInference/Source/NeuralNetworkInference/Public/NeuralNetwork.h
+++ b/Engine/Plugins/Experimental/NeuralNetworkInference/Source/NeuralNetworkInference/Public/NeuralNetwork.h
@@ -200,6 +200,9 @@ public:
 	const FNeuralTensor& GetInputTensor(const int32 InTensorIndex = 0) const;
 	const FNeuralTensor& GetOutputTensor(const int32 InTensorIndex = 0) const;
 
+	const FNeuralTensor& GetInputTensorForContext(const int32 InContextHandle, const int32 InTensorIndex = 0) const;
+	const FNeuralTensor& GetOutputTensorForContext(const int32 InContextHandle, const int32 InTensorIndex = 0) const;
+
 	/**
 	 * SetInputFromArrayCopy(), SetInputFromVoidPointerCopy(), and GetInputDataPointerMutable() are the only functions that allow modifying the
 	 * network input tensor(s) values:
@@ -212,6 +215,8 @@ public:
 	void SetInputFromVoidPointerCopy(const void* const InVoidPtr, const int32 InTensorIndex = 0);
 	void* GetInputDataPointerMutable(const int32 InTensorIndex = 0);
 
+	void* GetInputDataPointerMutableForContext(const int32 InContextHandle, const int32 InTensorIndex = 0);
+
 	/**
 	 * Non-computationally-efficient functions meant to be used only for debugging purposes, but should never be used on highly performant systems:
 	 * - InputTensorsToCPU copies the CPU memory of the desired input tensor(s) to GPU (to debug InputDeviceType == ENeuralDeviceType::GPU).
@@ -221,6 +226,24 @@ public:
 	void InputTensorsToGPU(const TArray<int32>& InTensorIndexes = TArray<int32>());
 	void OutputTensorsToCPU(const TArray<int32>& InTensorIndexes = TArray<int32>());
 
+	/**
+	 * Create an inference context. An inference context contains resources to run inference with the version of Run() that takes a context handle.
+	 * Each independent instance of the NeuralNetwork requires it's own inference context.
+	 * Returns a handle to the context. -1 is an invalid handle.
+	 */ 
+	int32 CreateInferenceContext();
+
+	/**
+	 * Destroy an inference context that was created with CreateInferenceContext()
+	 */
+	void DestroyInferenceContext(int32 ContextHandle);
+
+	/**
+	 * Run() executes the forward pass of the current UNeuralNetwork given an inference context created with CreateInferenceContext().
+	 * This should be called on the render thread.
+	 */
+	void Run(class FRDGBuilder& GraphBuilder, int32 ContextHandle);
+
 	/**
 	 * Load() + SetInputFromArrayCopy() + Run() is the simplest way to load an ONNX file, set the input tensor(s), and run inference on it. All other
 	 * functions in UNeuralNetwork provide additional functionality (set the device to CPU/GPU, more complicated but faster ways to set the input,
diff --git a/Engine/Plugins/Experimental/NeuralNetworkInference/Source/NeuralNetworkInference/Public/NeuralTensor.h b/Engine/Plugins/Experimental/NeuralNetworkInference/Source/NeuralNetworkInference/Public/NeuralTensor.h
--- a/Engine/Plugins/Experimental/NeuralNetworkInference/Source/NeuralNetworkInference/Public/NeuralTensor.h
+++ b/Engine/Plugins/Experimental/NeuralNetworkInference/Source/NeuralNetworkInference/Public/NeuralTensor.h
@@ -311,11 +311,19 @@ public:
 
 	/**
 	 * It allocates the GPU data, in particular the pooled buffer.
+	 * Should be called on the render thread. Internally this blocks for reading back work on the render thread.
 	 * @param InOutNativeResource Pointer to the ID3D12Resource obtained from PooledBuffer. This resource can be shared with the DirectML execution
 	 * provider.
 	 */
 	bool InitPooledBufferForUEAndORTBackEnd(void** InOutNativeResource = nullptr);
 
+	/**
+	 * It allocates the GPU data, in particular the pooled buffer.
+	 * Should be called on the render thread.
+	 * @param GraphBuilder. RDG graph builder to use for allocation.
+	 */
+	void InitPooledBufferForUEAndORTBackEnd_RenderThread(FRDGBuilder& GraphBuilder);
+
 	/**
 	 * It returns a constant reference of PooledBuffer. @see PooledBuffer for more details.
 	 */
