NNI: Ensure that we release DirectML allocations from inference contexts when destroying a NN.

diff --git a/Engine/Plugins/Experimental/NeuralNetworkInference/Source/NeuralNetworkInference/Private/NeuralNetworkImplBackEndUEAndORT.cpp b/Engine/Plugins/Experimental/NeuralNetworkInference/Source/NeuralNetworkInference/Private/NeuralNetworkImplBackEndUEAndORT.cpp
--- a/Engine/Plugins/Experimental/NeuralNetworkInference/Source/NeuralNetworkInference/Private/NeuralNetworkImplBackEndUEAndORT.cpp
+++ b/Engine/Plugins/Experimental/NeuralNetworkInference/Source/NeuralNetworkInference/Private/NeuralNetworkImplBackEndUEAndORT.cpp
@@ -1050,6 +1050,12 @@ void UNeuralNetwork::FImplBackEndUEAndORT::ClearResources()
 		}
 
 		DmlGPUResources.Reset(0);
+		
+		for (FInferenceContext& Context : Contexts)
+		{
+			Context.ReleaseDMLAllocations(DmlApi);
+		}
+		
 		DmlApi = nullptr;
 	}
 #endif //PLATFORM_WIN64
