Fix spurious NNI warnings during cook
Allow use of GPU by NNI in cooked builds

Unfortunately the code is a bit tortured due to the constraints of this unrefactored code base and the late date of the fix relative to Release-5.0 lockdown.  Cleanup will occur in Main

diff --git a/Engine/Plugins/Experimental/NeuralNetworkInference/Source/NeuralNetworkInference/Private/NeuralNetwork.cpp b/Engine/Plugins/Experimental/NeuralNetworkInference/Source/NeuralNetworkInference/Private/NeuralNetwork.cpp
--- a/Engine/Plugins/Experimental/NeuralNetworkInference/Source/NeuralNetworkInference/Private/NeuralNetwork.cpp
+++ b/Engine/Plugins/Experimental/NeuralNetworkInference/Source/NeuralNetworkInference/Private/NeuralNetwork.cpp
@@ -543,10 +543,15 @@ bool UNeuralNetwork::Load()
 	// UEAndORT
 	if (BackEndForCurrentPlatform == ENeuralBackEnd::UEAndORT)
 	{
-		UNeuralNetwork::FImplBackEndUEAndORT::WarnAndSetDeviceToCPUIfDX12NotEnabled(DeviceType, /*bShouldOpenMessageLog*/true);
+		// Need to catch this in a local because the back end isn't created yet
+		bool bIsCPUForced =
+			UNeuralNetwork::FImplBackEndUEAndORT::ForceCPUIfNoGPU(DeviceType);
 		bIsLoaded = UNeuralNetwork::FImplBackEndUEAndORT::Load(ImplBackEndUEAndORT, OnAsyncRunCompletedDelegate,
 			ThreadModeDelegateForAsyncRunCompleted, ResoucesCriticalSection, AreInputTensorSizesVariable, ModelReadFromFileInBytes,
 			ModelFullFilePath, GetDeviceType(), GetInputDeviceType(), GetOutputDeviceType());
+#ifdef WITH_UE_AND_ORT_SUPPORT
+		ImplBackEndUEAndORT->bIsCPUForced = bIsCPUForced;
+#endif
 	}
 	// UEOnly
 	else if (BackEndForCurrentPlatform == ENeuralBackEnd::UEOnly)
@@ -681,11 +686,6 @@ void UNeuralNetwork::PostLoad()
 	// If ModelReadFromFileInBytes is not empty, call Load() 
 	if (ModelReadFromFileInBytes.Num() > 0)
 	{
-		// If GPU selected but not compatible, set to CPU
-		if (BackEnd == ENeuralBackEnd::UEAndORT)
-		{
-			FImplBackEndUEAndORT::WarnAndSetDeviceToCPUIfDX12NotEnabled(DeviceType, /*bShouldOpenMessageLog*/false);
-		}
 		// Load
 		if (!Load())
 		{
@@ -703,7 +703,22 @@ void UNeuralNetwork::Serialize(FArchive& Archive)
 		ReimportAssetFromEditorData();
 	}
 #endif // WITH_EDITORONLY_DATA
-	Super::Serialize(Archive);
+#ifdef WITH_UE_AND_ORT_SUPPORT
+	bool bSwapDeviceType = (Archive.IsSaving() && (BackEndForCurrentPlatform == ENeuralBackEnd::UEAndORT) && (ImplBackEndUEAndORT->bIsCPUForced));
+#else
+	bool bSwapDeviceType = false;
+#endif
+
+	if (bSwapDeviceType)
+	{
+		DeviceType = ENeuralDeviceType::GPU;
+		Super::Serialize(Archive);
+		DeviceType = ENeuralDeviceType::CPU;
+	}
+	else
+	{
+		Super::Serialize(Archive);
+	}
 }
 
 
diff --git a/Engine/Plugins/Experimental/NeuralNetworkInference/Source/NeuralNetworkInference/Private/NeuralNetworkImplBackEndUEAndORT.cpp b/Engine/Plugins/Experimental/NeuralNetworkInference/Source/NeuralNetworkInference/Private/NeuralNetworkImplBackEndUEAndORT.cpp
--- a/Engine/Plugins/Experimental/NeuralNetworkInference/Source/NeuralNetworkInference/Private/NeuralNetworkImplBackEndUEAndORT.cpp
+++ b/Engine/Plugins/Experimental/NeuralNetworkInference/Source/NeuralNetworkInference/Private/NeuralNetworkImplBackEndUEAndORT.cpp
@@ -2,6 +2,7 @@
 
 #include "NeuralNetworkImplBackEndUEAndORT.h"
 #include "Async/Async.h"
+#include "NeuralEnumClasses.h"
 #include "NeuralNetworkInferenceUtils.h"
 #include "NeuralNetworkInferenceUtilsGPU.h"
 #include "RedirectCoutAndCerrToUeLog.h"
@@ -288,37 +289,17 @@ UNeuralNetwork::FImplBackEndUEAndORT::~FImplBackEndUEAndORT()
 /* FImplBackEndUEAndORT public functions
  *****************************************************************************/
 
-void UNeuralNetwork::FImplBackEndUEAndORT::WarnAndSetDeviceToCPUIfDX12NotEnabled(ENeuralDeviceType& InOutDeviceType, const bool bInShouldOpenMessageLog)
+bool UNeuralNetwork::FImplBackEndUEAndORT::ForceCPUIfNoGPU(ENeuralDeviceType& InOutDeviceType)
 {
 	if (InOutDeviceType != ENeuralDeviceType::CPU)
 	{
 		if (!IsGPUSupported())
 		{
 			InOutDeviceType = ENeuralDeviceType::CPU;
-
-			const FString RHIName = GDynamicRHI->GetName();
-#ifdef PLATFORM_WIN64
-			const FString ErrorMessage = TEXT("On Windows, only DirectX 12 rendering (\"D3D12\") is compatible with the UEAndORT back end of NeuralNetworkInference (NNI). Instead, \"")
-				+ RHIName + TEXT("\" was used. You have the following options:\n\n"
-					"\t1. (Recommended) Switch Unreal Engine to DX12. In order to do that:\n"
-					"\t\t - Go to \"Project Settings\", \"Platforms\", \"Windows\", \"Default RHI\".\n"
-					"\t\t - Select \"DirectX 12\".\n"
-					"\t\t - Restart Unreal Engine.\n"
-					"\t2. Alternatively, switch the network to CPU with UNeuralNetwork::SetDeviceType().\n\n"
-					"Network set to CPU provisionally.");
-#else //PLATFORM_WIN64
-			const FString ErrorMessage = TEXT("GPU version is not supported for non-Windows platforms yet. Switch the network to CPU with UNeuralNetwork::SetDeviceType() or run from Windows.\n\n"
-					"Network set to CPU provisionally.");
-#endif //PLATFORM_WIN64
-			UE_LOG(LogNeuralNetworkInference, Warning, TEXT("FImplBackEndUEAndORT::WarnAndSetDeviceToCPUIfDX12NotEnabled(): %s"), *ErrorMessage);
-#if WITH_EDITOR
-			if (bInShouldOpenMessageLog)
-			{
-				FMessageDialog::Open(EAppMsgType::Ok, FText::FromString(ErrorMessage));
-			}
-#endif //WITH_EDITOR
+			return true;
 		}
 	}
+	return false;
 }
 
 bool UNeuralNetwork::FImplBackEndUEAndORT::IsGPUSupported()
@@ -427,6 +408,11 @@ void UNeuralNetwork::FImplBackEndUEAndORT::Run(const ENeuralSynchronousMode InSy
 	const ENeuralDeviceType InInputDeviceType)
 {
 #ifdef WITH_UE_AND_ORT_SUPPORT
+	if (! bHasRun)
+	{
+		WarnOnCPUForced(true);
+		bHasRun = true;
+	}
 #if WITH_EDITOR
 	try
 #endif //WITH_EDITOR
@@ -456,8 +442,8 @@ void UNeuralNetwork::FImplBackEndUEAndORT::Run(const ENeuralSynchronousMode InSy
 #endif //WITH_EDITOR
 
 #else //WITH_UE_AND_ORT_SUPPORT
-	UE_LOG(LogNeuralNetworkInference, Warning, TEXT("FImplBackEndUEAndORT::Run(): Platform or Operating System not suported yet for UEAndORT"
-		" BackEnd. Set BackEnd to ENeuralBackEnd::Auto or ENeuralBackEnd::UEOnly for this platform."));
+		UE_LOG(LogNeuralNetworkInference, Warning, TEXT("FImplBackEndUEAndORT::Run(): Platform or Operating System not suported yet for UEAndORT"
+														" BackEnd. Set BackEnd to ENeuralBackEnd::Auto or ENeuralBackEnd::UEOnly for this platform."));
 #endif //WITH_UE_AND_ORT_SUPPORT
 }
 
@@ -467,6 +453,38 @@ void UNeuralNetwork::FImplBackEndUEAndORT::Run(const ENeuralSynchronousMode InSy
  *****************************************************************************/
 
 #ifdef WITH_UE_AND_ORT_SUPPORT
+
+void UNeuralNetwork::FImplBackEndUEAndORT::WarnOnCPUForced(bool bInShouldOpenMessageLog)
+{
+	if (!bIsCPUForced)
+	{
+		return;
+	}
+
+	const FString RHIName = GDynamicRHI->GetName();
+#ifdef PLATFORM_WIN64
+	const FString ErrorMessage = TEXT("On Windows, only DirectX 12 rendering (\"D3D12\") is compatible with the UEAndORT back end of NeuralNetworkInference (NNI). Instead, \"")
+		+ RHIName + TEXT("\" was used. You have the following options:\n\n"
+						 "\t1. (Recommended) Switch Unreal Engine to DX12. In order to do that:\n"
+						 "\t\t - Go to \"Project Settings\", \"Platforms\", \"Windows\", \"Default RHI\".\n"
+						 "\t\t - Select \"DirectX 12\".\n"
+						 "\t\t - Restart Unreal Engine.\n"
+						 "\t2. Alternatively, switch the network to CPU with UNeuralNetwork::SetDeviceType().\n\n"
+						 "Network set to CPU provisionally.");
+#else //PLATFORM_WIN64
+	const FString ErrorMessage = TEXT("GPU version is not supported for non-Windows platforms yet. Switch the network to CPU with UNeuralNetwork::SetDeviceType() or run from Windows.\n\n"
+									  "Network set to CPU provisionally.");
+#endif //PLATFORM_WIN64
+	UE_LOG(LogNeuralNetworkInference, Warning, TEXT("FImplBackEndUEAndORT::WarnAndSetDeviceToCPUIfDX12NotEnabled(): %s"), *ErrorMessage);
+#if WITH_EDITOR
+	if (bInShouldOpenMessageLog)
+	{
+		FMessageDialog::Open(EAppMsgType::Ok, FText::FromString(ErrorMessage));
+	}
+#endif //WITH_EDITOR
+}
+
+
 void UNeuralNetwork::FImplBackEndUEAndORT::EnsureAsyncTaskCompletion() const
 {
 	if (NeuralNetworkAsyncTask && !NeuralNetworkAsyncTask->IsDone())
diff --git a/Engine/Plugins/Experimental/NeuralNetworkInference/Source/NeuralNetworkInference/Private/NeuralNetworkImplBackEndUEAndORT.h b/Engine/Plugins/Experimental/NeuralNetworkInference/Source/NeuralNetworkInference/Private/NeuralNetworkImplBackEndUEAndORT.h
--- a/Engine/Plugins/Experimental/NeuralNetworkInference/Source/NeuralNetworkInference/Private/NeuralNetworkImplBackEndUEAndORT.h
+++ b/Engine/Plugins/Experimental/NeuralNetworkInference/Source/NeuralNetworkInference/Private/NeuralNetworkImplBackEndUEAndORT.h
@@ -47,8 +47,7 @@ public:
 
 	~FImplBackEndUEAndORT();
 
-	static void WarnAndSetDeviceToCPUIfDX12NotEnabled(ENeuralDeviceType& InOutDeviceType, const bool bInShouldOpenMessageLog);
-
+	[[nodiscard]] static bool ForceCPUIfNoGPU(ENeuralDeviceType& InOutDeviceType);
 	static bool IsGPUSupported();
 
 	static bool Load(TSharedPtr<FImplBackEndUEAndORT>& InOutImplBackEndUEAndORT, FOnAsyncRunCompleted& InOutOnAsyncRunCompletedDelegate,
@@ -93,6 +92,12 @@ private:
 	TArray<Ort::Value> OutputOrtTensors;
 	/** Tensor names */
 	TArray<const char*> OutputTensorNames;
+	public:
+	/** Has the device type been forced to CPU after requesting GPU*/
+	bool bIsCPUForced = false;
+	private:
+	/** Has this network been run before */
+	bool bHasRun = false;
 
 	/**
 	 * Storage for input and output tensors for ONNX backend.
@@ -171,5 +176,8 @@ private:
 
 	void ClearResources();
 
+	void WarnOnCPUForced(bool bInShouldOpenMessageLog);
+
+
 #endif //WITH_UE_AND_ORT_SUPPORT
 };
