EditorDomain and NeuralNetworkInference: Fix UNeuralNetwork::Serialize when called with Saving on the default object to discover CustomVersions.
#rnx

diff --git a/Engine/Plugins/Experimental/NeuralNetworkInference/Source/NeuralNetworkInference/Private/NeuralNetwork.cpp b/Engine/Plugins/Experimental/NeuralNetworkInference/Source/NeuralNetworkInference/Private/NeuralNetwork.cpp
--- a/Engine/Plugins/Experimental/NeuralNetworkInference/Source/NeuralNetworkInference/Private/NeuralNetwork.cpp
+++ b/Engine/Plugins/Experimental/NeuralNetworkInference/Source/NeuralNetworkInference/Private/NeuralNetwork.cpp
@@ -704,7 +704,8 @@ void UNeuralNetwork::Serialize(FArchive& Archive)
 	}
 #endif // WITH_EDITORONLY_DATA
 #ifdef WITH_UE_AND_ORT_SUPPORT
-	bool bSwapDeviceType = (Archive.IsSaving() && (BackEndForCurrentPlatform == ENeuralBackEnd::UEAndORT) && (ImplBackEndUEAndORT->bIsCPUForced));
+	bool bSwapDeviceType = (Archive.IsSaving() && (BackEndForCurrentPlatform == ENeuralBackEnd::UEAndORT) &&
+		(ImplBackEndUEAndORT && ImplBackEndUEAndORT->bIsCPUForced));
 #else
 	bool bSwapDeviceType = false;
 #endif
