Remove MLRigDeformer Network from NNI Unit Tests

diff --git a/Engine/Plugins/Experimental/NeuralNetworkInference/Source/Private/NeuralNetworkInferenceQA/Private/ModelUnitTester.cpp b/Engine/Plugins/Experimental/NeuralNetworkInference/Source/Private/NeuralNetworkInferenceQA/Private/ModelUnitTester.cpp
--- a/Engine/Plugins/Experimental/NeuralNetworkInference/Source/Private/NeuralNetworkInferenceQA/Private/ModelUnitTester.cpp
+++ b/Engine/Plugins/Experimental/NeuralNetworkInference/Source/Private/NeuralNetworkInferenceQA/Private/ModelUnitTester.cpp
@@ -73,7 +73,6 @@ bool FModelUnitTester::GlobalTest(const FString& InProjectContentDir, const FStr
 
 	// Model names, input values, and number of repetitions for profiling
 	const TArray<FString> ModelNames({
-		TEXT("MLRigDeformer"),
 		TEXT("cloth_network"),
 		TEXT("HS"),
 		TEXT("RL")
@@ -81,14 +80,12 @@ bool FModelUnitTester::GlobalTest(const FString& InProjectContentDir, const FStr
 
 	// Ground truths
 	const TArray<TArray<double>> CPUGroundTruths({
-		{3.728547, 0.008774, 4.595651, 212.193216, 742.434561, 4.250668, 4.717748},
 		{0.042571, 0.023693, 0.015783, 13.100505, 8.050994, 0.028807, 0.016387},
 		{138.372906, 126.753839, 127.287254, 130.316062, 127.303424, 124.800896, 126.546051},
 		{0.488662, 0.472437, 0.478862, 0.522685, 0.038322, 0.480848, 0.483821}
 	});
 	
 	const TArray<TArray<double>> GPUGroundTruths({
-		{3.728547, 0.008774, 4.595651, 212.193208, 742.434578, 4.250668, 4.717748},
 		{0.042571, 0.023693, 0.015783, 13.100504, 8.050994, 0.028807, 0.016387},
 		{138.373184, 126.754100, 127.287398, 130.316194, 127.303495, 124.801134, 126.5462530},
 		{0.488662, 0.472437, 0.478862, 0.522685, 0.038322, 0.480848, 0.483821}
@@ -106,23 +103,23 @@ bool FModelUnitTester::GlobalTest(const FString& InProjectContentDir, const FStr
 	
 	// Speed profiling test - 0 repetitions means that test will not be run
 #ifdef WITH_UE_AND_ORT_SUPPORT
-	const TArray<int32> CPURepetitionsForUEAndORTBackEnd({ 1000, 1000,  50, 1000 });
+	const TArray<int32> CPURepetitionsForUEAndORTBackEnd({1000,  50, 1000 });
 #ifdef PLATFORM_WIN64
-	const TArray<int32> GPURepetitionsForUEAndORTBackEnd({ 1000, 1000, 100, 1000 });
+	const TArray<int32> GPURepetitionsForUEAndORTBackEnd({1000, 100, 1000 });
 #else //PLATFORM_WIN64
 	UE_LOG(LogNeuralNetworkInferenceQA, Display, TEXT("FModelUnitTester::GlobalTest(): GPU tests disabled for non-Windows platforms."));
-	const TArray<int32> GPURepetitionsForUEAndORTBackEnd({ 0, 0, 0, 0 });
+	const TArray<int32> GPURepetitionsForUEAndORTBackEnd({ 0, 0, 0 });
 #endif //PLATFORM_WIN64
 #else //WITH_UE_AND_ORT_SUPPORT
-	const TArray<int32> CPURepetitionsForUEAndORTBackEnd({ 0, 0, 0, 0 });
-	const TArray<int32> GPURepetitionsForUEAndORTBackEnd({ 0, 0, 0, 0 });
+	const TArray<int32> CPURepetitionsForUEAndORTBackEnd({ 0, 0, 0 });
+	const TArray<int32> GPURepetitionsForUEAndORTBackEnd({ 0, 0, 0 });
 #endif //WITH_UE_AND_ORT_SUPPORT
 #if WITH_EDITOR
-	const TArray<int32> CPURepetitionsForUEOnlyBackEnd({ 10, 0, 0, 0 });
-	const TArray<int32> GPURepetitionsForUEOnlyBackEnd({ 10, 10, 0, 0 });
+	const TArray<int32> CPURepetitionsForUEOnlyBackEnd({ 0, 0, 0 });
+	const TArray<int32> GPURepetitionsForUEOnlyBackEnd({ 10, 0, 0 });
 #else //WITH_EDITOR
-	const TArray<int32> CPURepetitionsForUEOnlyBackEnd({ 0, 0, 0, 0 });
-	const TArray<int32> GPURepetitionsForUEOnlyBackEnd({ 0, 0, 0, 0 });
+	const TArray<int32> CPURepetitionsForUEOnlyBackEnd({ 0, 0, 0 });
+	const TArray<int32> GPURepetitionsForUEOnlyBackEnd({ 0, 0, 0 });
 #endif //WITH_EDITOR
 	// Run tests
 	UE_LOG(LogNeuralNetworkInferenceQA, Display,
